Fix screenpercentage on Morpheus

diff --git a/Engine/Shaders/ClearReplacementShaders.usf b/Engine/Shaders/ClearReplacementShaders.usf
--- a/Engine/Shaders/ClearReplacementShaders.usf
+++ b/Engine/Shaders/ClearReplacementShaders.usf
@@ -24,6 +24,7 @@ float4 ClearPS() : SV_Target0
 
 #if FEATURE_LEVEL >= FEATURE_LEVEL_SM5
 RWTexture2D<float4> ClearTextureRW;
+uint4 TargetBounds; //xy upperleft, zw lowerright
 
 [numthreads(8,8,1)]
 void ClearTexture2DCS(uint2 Position : SV_DispatchThreadID)
@@ -31,6 +32,17 @@ void ClearTexture2DCS(uint2 Position : SV_DispatchThreadID)
 	ClearTextureRW[Position] = ClearColor;
 }
 
+//clear with emulated scissor rects.
+[numthreads(8, 8, 1)]
+void ClearTexture2DScissorCS(uint2 Position : SV_DispatchThreadID)
+{
+	uint2 FinalPos = Position + TargetBounds.xy;
+	if (all(FinalPos < TargetBounds.zw))
+	{
+		ClearTextureRW[FinalPos] = ClearColor;
+	}
+}
+
 
 uint ClearDword;
 RWBuffer<uint> ClearBufferRW;
diff --git a/Engine/Source/Runtime/Renderer/Private/PostProcess/PostProcessing.cpp b/Engine/Source/Runtime/Renderer/Private/PostProcess/PostProcessing.cpp
--- a/Engine/Source/Runtime/Renderer/Private/PostProcess/PostProcessing.cpp
+++ b/Engine/Source/Runtime/Renderer/Private/PostProcess/PostProcessing.cpp
@@ -1509,6 +1509,7 @@ void FPostProcessing::Process(FRHICommandListImmediate& RHICmdList, FViewInfo& V
 		AddGBufferVisualizationOverview(Context, SeparateTranslucency);
 
 		bool bStereoRenderingAndHMD = View.Family->EngineShowFlags.StereoRendering && View.Family->EngineShowFlags.HMDDistortion;
+		bool bHMDWantsUpscale = false;
 		if (bStereoRenderingAndHMD)
 		{
 			FRenderingCompositePass* Node = NULL;
@@ -1517,12 +1518,15 @@ void FPostProcessing::Process(FRHICommandListImmediate& RHICmdList, FViewInfo& V
 			{
 				Node = Context.Graph.RegisterPass(new FRCPassPostProcessHMD());
 			}
-#if MORPHEUS_ENGINE_DISTORTION
+
 			else if(DeviceType == EHMDDeviceType::DT_Morpheus)
 			{
+#if MORPHEUS_ENGINE_DISTORTION
 				Node = Context.Graph.RegisterPass(new FRCPassPostProcessMorpheus());
-			}
 #endif
+				bHMDWantsUpscale = true;
+			}
+
 			
 			if(Node)
 			{
@@ -1561,8 +1565,8 @@ void FPostProcessing::Process(FRHICommandListImmediate& RHICmdList, FViewInfo& V
 		}
 
 		// Do not use upscale if SeparateRenderTarget is in use!
-		if ((Cylinder > 0.01f || View.UnscaledViewRect != View.ViewRect)
-			&& (!View.Family->EngineShowFlags.StereoRendering || (!View.Family->EngineShowFlags.HMDDistortion && !View.Family->bUseSeparateRenderTarget)))
+		if ((Cylinder > 0.01f || View.UnscaledViewRect != View.ViewRect) && 
+			(bHMDWantsUpscale ||!View.Family->EngineShowFlags.StereoRendering || (!View.Family->EngineShowFlags.HMDDistortion && !View.Family->bUseSeparateRenderTarget)))
 		{
 			int32 UpscaleQuality = CVarUpscaleQuality.GetValueOnRenderThread();
 			UpscaleQuality = FMath::Clamp(UpscaleQuality, 0, 3);
diff --git a/Engine/Source/Runtime/UtilityShaders/Private/ClearReplacementShaders.cpp b/Engine/Source/Runtime/UtilityShaders/Private/ClearReplacementShaders.cpp
--- a/Engine/Source/Runtime/UtilityShaders/Private/ClearReplacementShaders.cpp
+++ b/Engine/Source/Runtime/UtilityShaders/Private/ClearReplacementShaders.cpp
@@ -9,6 +9,7 @@ IMPLEMENT_SHADER_TYPE(,FClearReplacementVS,TEXT("ClearReplacementShaders"),TEXT(
 IMPLEMENT_SHADER_TYPE(,FClearReplacementPS,TEXT("ClearReplacementShaders"),TEXT("ClearPS"),SF_Pixel);
 
 IMPLEMENT_SHADER_TYPE(,FClearTexture2DReplacementCS,TEXT("ClearReplacementShaders"),TEXT("ClearTexture2DCS"),SF_Compute);
+IMPLEMENT_SHADER_TYPE(,FClearTexture2DReplacementScissorCS,TEXT("ClearReplacementShaders"),TEXT("ClearTexture2DScissorCS"), SF_Compute);
 
 IMPLEMENT_SHADER_TYPE(,FClearBufferReplacementCS,TEXT("ClearReplacementShaders"),TEXT("ClearBufferCS"),SF_Compute);
 
@@ -19,6 +20,14 @@ void FClearTexture2DReplacementCS::SetParameters( FRHICommandList& RHICmdList, F
 	RHICmdList.SetUAVParameter(ComputeShaderRHI, ClearTextureRW.GetBaseIndex(), TextureRW);
 }
 
+void FClearTexture2DReplacementScissorCS::SetParameters(FRHICommandList& RHICmdList, FUnorderedAccessViewRHIParamRef TextureRW, FLinearColor InClearColor, FVector4 InTargetBounds)
+{
+	FComputeShaderRHIParamRef ComputeShaderRHI = GetComputeShader();
+	SetShaderValue(RHICmdList, ComputeShaderRHI, ClearColor, InClearColor);
+	SetShaderValue(RHICmdList, ComputeShaderRHI, TargetBounds, InTargetBounds);
+	RHICmdList.SetUAVParameter(ComputeShaderRHI, ClearTextureRW.GetBaseIndex(), TextureRW);
+}
+
 void FClearBufferReplacementCS::SetParameters( FRHICommandList& RHICmdList, FUnorderedAccessViewRHIParamRef BufferRW, uint32 Dword )
 {
 	FComputeShaderRHIParamRef ComputeShaderRHI = GetComputeShader();
diff --git a/Engine/Source/Runtime/UtilityShaders/Public/ClearReplacementShaders.h b/Engine/Source/Runtime/UtilityShaders/Public/ClearReplacementShaders.h
--- a/Engine/Source/Runtime/UtilityShaders/Public/ClearReplacementShaders.h
+++ b/Engine/Source/Runtime/UtilityShaders/Public/ClearReplacementShaders.h
@@ -105,6 +105,55 @@ protected:
 	FShaderResourceParameter ClearTextureRW;
 };
 
+class FClearTexture2DReplacementScissorCS : public FGlobalShader
+{
+	DECLARE_EXPORTED_SHADER_TYPE(FClearTexture2DReplacementScissorCS, Global, UTILITYSHADERS_API);
+public:
+	FClearTexture2DReplacementScissorCS() {}
+	FClearTexture2DReplacementScissorCS(const ShaderMetaType::CompiledShaderInitializerType& Initializer)
+		: FGlobalShader(Initializer)
+	{
+		ClearColor.Bind(Initializer.ParameterMap, TEXT("ClearColor"), SPF_Mandatory);
+		TargetBounds.Bind(Initializer.ParameterMap, TEXT("TargetBounds"), SPF_Mandatory);
+		ClearTextureRW.Bind(Initializer.ParameterMap, TEXT("ClearTextureRW"), SPF_Mandatory);
+	}
+
+	// FShader interface.
+	virtual bool Serialize(FArchive& Ar) override
+	{
+		bool bShaderHasOutdatedParameters = FGlobalShader::Serialize(Ar);
+		Ar << ClearColor << TargetBounds << ClearTextureRW;
+		return bShaderHasOutdatedParameters;
+	}
+
+	void SetParameters(FRHICommandList& RHICmdList, FUnorderedAccessViewRHIParamRef TextureRW, FLinearColor ClearColor, FVector4 TargetBounds);
+
+	static bool ShouldCache(EShaderPlatform Platform)
+	{
+		return IsFeatureLevelSupported(Platform, ERHIFeatureLevel::SM5);
+	}
+
+	const FShaderParameter& GetClearColorParameter()
+	{
+		return ClearColor;
+	}
+
+	const FShaderParameter& GetTargetBoundsParameter()
+	{
+		return TargetBounds;
+	}
+
+	const FShaderResourceParameter& GetClearTextureRWParameter()
+	{
+		return ClearTextureRW;
+	}
+
+protected:
+	FShaderParameter ClearColor;
+	FShaderParameter TargetBounds;
+	FShaderResourceParameter ClearTextureRW;
+};
+
 class FClearBufferReplacementCS : public FGlobalShader
 {
 	DECLARE_EXPORTED_SHADER_TYPE(FClearBufferReplacementCS, Global, UTILITYSHADERS_API);
