#UE4 CIS fix. Fixed an alignment issue when passing FVector4 by value.

diff --git a/Engine/Source/Runtime/UtilityShaders/Private/ClearReplacementShaders.cpp b/Engine/Source/Runtime/UtilityShaders/Private/ClearReplacementShaders.cpp
--- a/Engine/Source/Runtime/UtilityShaders/Private/ClearReplacementShaders.cpp
+++ b/Engine/Source/Runtime/UtilityShaders/Private/ClearReplacementShaders.cpp
@@ -20,7 +20,7 @@ void FClearTexture2DReplacementCS::SetParameters( FRHICommandList& RHICmdList, F
 	RHICmdList.SetUAVParameter(ComputeShaderRHI, ClearTextureRW.GetBaseIndex(), TextureRW);
 }
 
-void FClearTexture2DReplacementScissorCS::SetParameters(FRHICommandList& RHICmdList, FUnorderedAccessViewRHIParamRef TextureRW, FLinearColor InClearColor, FVector4 InTargetBounds)
+void FClearTexture2DReplacementScissorCS::SetParameters(FRHICommandList& RHICmdList, FUnorderedAccessViewRHIParamRef TextureRW, FLinearColor InClearColor, const FVector4& InTargetBounds)
 {
 	FComputeShaderRHIParamRef ComputeShaderRHI = GetComputeShader();
 	SetShaderValue(RHICmdList, ComputeShaderRHI, ClearColor, InClearColor);
diff --git a/Engine/Source/Runtime/UtilityShaders/Public/ClearReplacementShaders.h b/Engine/Source/Runtime/UtilityShaders/Public/ClearReplacementShaders.h
--- a/Engine/Source/Runtime/UtilityShaders/Public/ClearReplacementShaders.h
+++ b/Engine/Source/Runtime/UtilityShaders/Public/ClearReplacementShaders.h
@@ -126,7 +126,7 @@ public:
 		return bShaderHasOutdatedParameters;
 	}
 
-	void SetParameters(FRHICommandList& RHICmdList, FUnorderedAccessViewRHIParamRef TextureRW, FLinearColor ClearColor, FVector4 TargetBounds);
+	void SetParameters(FRHICommandList& RHICmdList, FUnorderedAccessViewRHIParamRef TextureRW, FLinearColor ClearColor, const FVector4& InTargetBounds);
 
 	static bool ShouldCache(EShaderPlatform Platform)
 	{
