Fix non-unity build.

diff --git a/Engine/Source/Runtime/UtilityShaders/Private/OneColorShader.cpp b/Engine/Source/Runtime/UtilityShaders/Private/OneColorShader.cpp
--- a/Engine/Source/Runtime/UtilityShaders/Private/OneColorShader.cpp
+++ b/Engine/Source/Runtime/UtilityShaders/Private/OneColorShader.cpp
@@ -3,6 +3,7 @@
 
 #include "UtilityShadersPrivatePCH.h"
 #include "OneColorShader.h"
+#include "ShaderParameterUtils.h"
 
 BEGIN_UNIFORM_BUFFER_STRUCT(FClearShaderUB, )
 DECLARE_UNIFORM_BUFFER_STRUCT_MEMBER_ARRAY(FVector4, DrawColorMRT, [8] )
