Switch to Decoder2 and GPU conversion to RGBA for PS4 MediaFramework.

diff --git a/Engine/Shaders/MediaShaders.usf b/Engine/Shaders/MediaShaders.usf
new file mode 100644
--- /dev/null
+++ b/Engine/Shaders/MediaShaders.usf
@@ -0,0 +1,48 @@
+// Copyright 1998-2015 Epic Games, Inc. All Rights Reserved.
+
+/*=============================================================================
+MediaShaders.usf: Shaders to convert various media decoder outputs to RGBA
+=============================================================================*/
+
+#include "Common.usf"
+
+/*=============================================================================
+* Vertex shader
+=============================================================================*/
+
+struct FMediaShaderVertexOutput
+{
+	float4 Position : SV_POSITION;
+	float2 UV : TEXCOORD0;
+};
+
+void MainVertexShader(
+	in float4 InPosition : ATTRIBUTE0,
+	in float2 InUV : ATTRIBUTE1,
+	out FMediaShaderVertexOutput Output
+	)
+{
+	Output.Position = InPosition;
+	Output.UV = InUV;
+}
+
+/*=============================================================================
+* Pixel shader
+=============================================================================*/
+
+//shader to convert from Full-range YCbCr to RGBA
+void YCbCrConvertPS(
+	float2 InUV : TEXCOORD0,
+	out float4 OutColor : SV_Target0)
+{
+	float3 YCbCr;
+	YCbCr.x = YCbCrConvertUB.LumaTexture.Sample(YCbCrConvertUB.YCbCrSampler, InUV.xy).x - 0.0625; //e.g. 16 / 255
+	YCbCr.yz = YCbCrConvertUB.CbCrTexture.Sample(YCbCrConvertUB.YCbCrSampler, InUV.xy).xy - float2(0.5, 0.5);
+			
+	
+	OutColor =	float4(
+					dot(float3(1.1644f, 0.0f, 1.7927f), YCbCr), // R
+					dot(float3(1.1644f, -0.2133f, -0.5329f), YCbCr), // G
+					dot(float3(1.1644f, 2.1124f, 0.0f), YCbCr), // B
+					1.0f);
+}
\ No newline at end of file
diff --git a/Engine/Source/Runtime/Media/Public/IMediaStream.h b/Engine/Source/Runtime/Media/Public/IMediaStream.h
--- a/Engine/Source/Runtime/Media/Public/IMediaStream.h
+++ b/Engine/Source/Runtime/Media/Public/IMediaStream.h
@@ -33,7 +33,7 @@ public:
 	/**
 	 * Adds the given media sink to this stream.
 	 *
-	 * @param Sink The sink to add.
+	 * @param Sink The sink to add.  Uses deferred CPU copies to update the resource.
 	 * @see RemoveSink
 	 */
 	virtual void AddSink(const TSharedRef<IMediaSink, ESPMode::ThreadSafe>& Sink) = 0;
diff --git a/Engine/Source/Runtime/Media/Public/IMediaVideoTrack.h b/Engine/Source/Runtime/Media/Public/IMediaVideoTrack.h
--- a/Engine/Source/Runtime/Media/Public/IMediaVideoTrack.h
+++ b/Engine/Source/Runtime/Media/Public/IMediaVideoTrack.h
@@ -44,6 +44,22 @@ public:
 	 */
 	virtual IMediaStream& GetStream() = 0;
 
+#if WITH_ENGINE
+	/**
+	* Binds the given TextureResource to the track.  The MediaPlayer/Track implementation can use this to bypass the normal CPU copy deferred sinks
+	* and directly update the resource with the GPU if desired.
+	*
+	* @see RemoveBoundTexture
+	*/
+	virtual void AddBoundTexture(class FRHITexture* BoundResource) { };
+
+	/**
+	* Removes the given texture from the track binding.  Texture will no longer be updated via BoundTexture path.
+	* @see AddBoundTexture
+	*/
+	virtual void RemoveBoundTexture(class FRHITexture* BoundResource) { };
+#endif
+
 public:
 
 	/**
diff --git a/Engine/Source/Runtime/MediaAssets/Private/Assets/MediaTexture.cpp b/Engine/Source/Runtime/MediaAssets/Private/Assets/MediaTexture.cpp
--- a/Engine/Source/Runtime/MediaAssets/Private/Assets/MediaTexture.cpp
+++ b/Engine/Source/Runtime/MediaAssets/Private/Assets/MediaTexture.cpp
@@ -80,6 +80,17 @@ float UMediaTexture::GetSurfaceHeight() const
 	return CachedDimensions.Y;
 }
 
+void UMediaTexture::UpdateResource()
+{
+#if WITH_ENGINE
+	if (VideoTrack.IsValid() && Resource)
+	{		
+		VideoTrack->RemoveBoundTexture(Resource->TextureRHI.GetReference());
+	}
+#endif
+	UTexture::UpdateResource();
+}
+
 
 /* UObject  overrides
  *****************************************************************************/
@@ -191,6 +202,10 @@ void UMediaTexture::InitializeTrack()
 	if (VideoTrack.IsValid())
 	{
 		VideoTrack->GetStream().RemoveSink(VideoBuffer);
+
+#if WITH_ENGINE
+		VideoTrack->RemoveBoundTexture(Resource->TextureRHI.GetReference());
+#endif
 		VideoTrack.Reset();
 	}
 
@@ -230,6 +245,12 @@ void UMediaTexture::InitializeTrack()
 	if (VideoTrack.IsValid())
 	{
 		VideoTrack->GetStream().AddSink(VideoBuffer);
+
+#if WITH_ENGINE
+		FlushRenderingCommands();
+		IMediaVideoTrack* VideoTrackPtr = static_cast<IMediaVideoTrack*>(VideoTrack.Get());
+		VideoTrackPtr->AddBoundTexture((Resource->TextureRHI.GetReference()));
+#endif
 	}
 }
 
diff --git a/Engine/Source/Runtime/MediaAssets/Public/MediaTexture.h b/Engine/Source/Runtime/MediaAssets/Public/MediaTexture.h
--- a/Engine/Source/Runtime/MediaAssets/Public/MediaTexture.h
+++ b/Engine/Source/Runtime/MediaAssets/Public/MediaTexture.h
@@ -99,6 +99,7 @@ public:
 	virtual EMaterialValueType GetMaterialType() override;
 	virtual float GetSurfaceWidth() const override;
 	virtual float GetSurfaceHeight() const override;
+	virtual void UpdateResource() override;	
 
 public:
 
diff --git a/Engine/Source/Runtime/UtilityShaders/Private/MediaShaders.cpp b/Engine/Source/Runtime/UtilityShaders/Private/MediaShaders.cpp
new file mode 100644
--- /dev/null
+++ b/Engine/Source/Runtime/UtilityShaders/Private/MediaShaders.cpp
@@ -0,0 +1,28 @@
+// Copyright 1998-2015 Epic Games, Inc. All Rights Reserved.
+
+#include "UtilityShadersPrivatePCH.h"
+#include "MediaShaders.h"
+#include "ShaderParameterUtils.h"
+
+BEGIN_UNIFORM_BUFFER_STRUCT(FYCbCrConvertUB, )
+DECLARE_UNIFORM_BUFFER_STRUCT_MEMBER_TEXTURE(Texture2D, LumaTexture)
+DECLARE_UNIFORM_BUFFER_STRUCT_MEMBER_TEXTURE(Texture2D, CbCrTexture)
+DECLARE_UNIFORM_BUFFER_STRUCT_MEMBER_SAMPLER(SamplerState, YCbCrSampler)
+END_UNIFORM_BUFFER_STRUCT(FYCbCrConvertUB)
+
+IMPLEMENT_UNIFORM_BUFFER_STRUCT(FYCbCrConvertUB, TEXT("YCbCrConvertUB"));
+
+IMPLEMENT_SHADER_TYPE(, FMedaShadersVS, TEXT("MediaShaders"), TEXT("MainVertexShader"), SF_Vertex);
+
+IMPLEMENT_SHADER_TYPE(, FYCbCrConvertShaderPS, TEXT("MediaShaders"), TEXT("YCbCrConvertPS"), SF_Pixel);
+
+void FYCbCrConvertShaderPS::SetParameters(FRHICommandList& RHICmdList, FTextureRHIRef LumaTexture, FTextureRHIRef CbCrTexture)
+{
+	FYCbCrConvertUB YCbCrUB;
+	YCbCrUB.LumaTexture = LumaTexture;
+	YCbCrUB.CbCrTexture = CbCrTexture;
+	YCbCrUB.YCbCrSampler = TStaticSamplerState<SF_Bilinear>::GetRHI();		
+
+	TUniformBufferRef<FYCbCrConvertUB> Data = TUniformBufferRef<FYCbCrConvertUB>::CreateUniformBufferImmediate(YCbCrUB, UniformBuffer_SingleFrame);
+	SetUniformBufferParameter(RHICmdList, GetPixelShader(), GetUniformBufferParameter<FYCbCrConvertUB>(), Data);	
+}
\ No newline at end of file
diff --git a/Engine/Source/Runtime/UtilityShaders/Public/MediaShaders.h b/Engine/Source/Runtime/UtilityShaders/Public/MediaShaders.h
new file mode 100644
--- /dev/null
+++ b/Engine/Source/Runtime/UtilityShaders/Public/MediaShaders.h
@@ -0,0 +1,84 @@
+// Copyright 1998-2015 Epic Games, Inc. All Rights Reserved.
+
+/*================================================================================
+MediaShaders.h: Shaders for converting decoded media files to RGBA format
+================================================================================*/
+
+#include "Engine/EngineTypes.h"
+#include "RenderCore.h"
+#include "GlobalShader.h"
+#include "ShaderParameters.h"
+#include "ShaderParameterUtils.h"
+#include "RHIStaticStates.h"
+
+#pragma once
+
+class FMediaShaderHelper
+{
+	//draws a simple full screen quad with 0-1 UVs.
+	//user must set up shaders and global shader state before calling.
+	static void DrawFullScreenQuad();
+};
+
+class FMedaShadersVS : public FGlobalShader
+{
+	DECLARE_SHADER_TYPE(FMedaShadersVS, Global);
+
+public:
+
+	static bool ShouldCache(EShaderPlatform Platform)
+	{
+		return IsFeatureLevelSupported(Platform, ERHIFeatureLevel::SM4);
+	}
+
+	/** Default constructor. */
+	FMedaShadersVS() {}
+
+	/** Initialization constructor. */
+	FMedaShadersVS(const ShaderMetaType::CompiledShaderInitializerType& Initializer)
+		: FGlobalShader(Initializer)
+	{
+	}
+
+	virtual bool Serialize(FArchive& Ar) override
+	{
+		bool bShaderHasOutdatedParameters = FGlobalShader::Serialize(Ar);
+		return bShaderHasOutdatedParameters;
+	}	
+
+private:
+};
+
+/** Shader to convert a full-res luma (Y) and half-res chroma (CbCr) texture to RGBA */
+class FYCbCrConvertShaderPS : public FGlobalShader
+{
+	DECLARE_SHADER_TYPE(FYCbCrConvertShaderPS, Global);
+
+public:
+
+	static bool ShouldCache(EShaderPlatform Platform)
+	{
+		return IsFeatureLevelSupported(Platform, ERHIFeatureLevel::SM4);
+	}
+
+	/** Default constructor. */
+	FYCbCrConvertShaderPS() {}
+
+	/** Initialization constructor. */
+	FYCbCrConvertShaderPS(const ShaderMetaType::CompiledShaderInitializerType& Initializer)
+		: FGlobalShader(Initializer)
+	{		
+	}
+
+	virtual bool Serialize(FArchive& Ar) override
+	{
+		bool bShaderHasOutdatedParameters = FGlobalShader::Serialize(Ar);		
+		return bShaderHasOutdatedParameters;
+	}
+
+	void SetParameters(FRHICommandList& RHICmdList, FTextureRHIRef LumaTexture, FTextureRHIRef CbCrTexture);
+
+private:
+
+	
+};
\ No newline at end of file
