OpenGL: Fix YUV->RGBA conversion (MediaShaders).

- Cannot reuse the sampler.

diff --git a/Engine/Shaders/MediaShaders.usf b/Engine/Shaders/MediaShaders.usf
--- a/Engine/Shaders/MediaShaders.usf
+++ b/Engine/Shaders/MediaShaders.usf
@@ -35,9 +35,10 @@ void YCbCrConvertPS(
 	float2 InUV : TEXCOORD0,
 	out float4 OutColor : SV_Target0)
 {
+	OutColor = float4(0, 0, 0, 0);
 	float3 YCbCr;
-	YCbCr.x = YCbCrConvertUB.LumaTexture.Sample(YCbCrConvertUB.YCbCrSampler, InUV.xy).x - 0.0625; //e.g. 16 / 255
-	YCbCr.yz = YCbCrConvertUB.CbCrTexture.Sample(YCbCrConvertUB.YCbCrSampler, InUV.xy).xy - float2(0.5, 0.5);
+	YCbCr.x = YCbCrConvertUB.LumaTexture.Sample(YCbCrConvertUB.LumaSampler, InUV.xy).x - 0.0625; //e.g. 16 / 255
+	YCbCr.yz = YCbCrConvertUB.CbCrTexture.Sample(YCbCrConvertUB.CbCrSampler, InUV.xy).xy - float2(0.5, 0.5);
 			
 	
 	OutColor =	float4(
diff --git a/Engine/Source/Runtime/UtilityShaders/Private/MediaShaders.cpp b/Engine/Source/Runtime/UtilityShaders/Private/MediaShaders.cpp
--- a/Engine/Source/Runtime/UtilityShaders/Private/MediaShaders.cpp
+++ b/Engine/Source/Runtime/UtilityShaders/Private/MediaShaders.cpp
@@ -7,7 +7,8 @@
 BEGIN_UNIFORM_BUFFER_STRUCT(FYCbCrConvertUB, )
 DECLARE_UNIFORM_BUFFER_STRUCT_MEMBER_TEXTURE(Texture2D, LumaTexture)
 DECLARE_UNIFORM_BUFFER_STRUCT_MEMBER_TEXTURE(Texture2D, CbCrTexture)
-DECLARE_UNIFORM_BUFFER_STRUCT_MEMBER_SAMPLER(SamplerState, YCbCrSampler)
+DECLARE_UNIFORM_BUFFER_STRUCT_MEMBER_SAMPLER(SamplerState, LumaSampler)
+DECLARE_UNIFORM_BUFFER_STRUCT_MEMBER_SAMPLER(SamplerState, CbCrSampler)
 END_UNIFORM_BUFFER_STRUCT(FYCbCrConvertUB)
 
 IMPLEMENT_UNIFORM_BUFFER_STRUCT(FYCbCrConvertUB, TEXT("YCbCrConvertUB"));
@@ -21,8 +22,9 @@ void FYCbCrConvertShaderPS::SetParameters(FRHICommandList& RHICmdList, FTextureR
 	FYCbCrConvertUB YCbCrUB;
 	YCbCrUB.LumaTexture = LumaTexture;
 	YCbCrUB.CbCrTexture = CbCrTexture;
-	YCbCrUB.YCbCrSampler = TStaticSamplerState<SF_Bilinear>::GetRHI();		
+	YCbCrUB.LumaSampler = TStaticSamplerState<SF_Bilinear>::GetRHI();
+	YCbCrUB.CbCrSampler = TStaticSamplerState<SF_Bilinear>::GetRHI();
 
 	TUniformBufferRef<FYCbCrConvertUB> Data = TUniformBufferRef<FYCbCrConvertUB>::CreateUniformBufferImmediate(YCbCrUB, UniformBuffer_SingleFrame);
 	SetUniformBufferParameter(RHICmdList, GetPixelShader(), GetUniformBufferParameter<FYCbCrConvertUB>(), Data);	
-}
\ No newline at end of file
+}
