UE4 - D3D12 - Fix UB check

diff --git a/Engine/Source/Runtime/UtilityShaders/Private/OneColorShader.cpp b/Engine/Source/Runtime/UtilityShaders/Private/OneColorShader.cpp
--- a/Engine/Source/Runtime/UtilityShaders/Private/OneColorShader.cpp
+++ b/Engine/Source/Runtime/UtilityShaders/Private/OneColorShader.cpp
@@ -30,8 +30,9 @@ void FOneColorPS::SetColors(FRHICommandList& RHICmdList, const FLinearColor* Col
 				ClearData.DrawColorMRT[i].W = Colors[i].A;
 			}
 
-			TUniformBufferRef<FClearShaderUB> Data = TUniformBufferRef<FClearShaderUB>::CreateUniformBufferImmediate(ClearData, UniformBuffer_SingleFrame);			
-			SetUniformBufferParameter(RHICmdList, GetPixelShader(), GetUniformBufferParameter<FClearShaderUB>(), Data);
+			FLocalUniformBuffer LocalUB = TUniformBufferRef<FClearShaderUB>::CreateLocalUniformBuffer(RHICmdList, ClearData, UniformBuffer_SingleFrame);	
+			auto& Parameter = GetUniformBufferParameter<FClearShaderUB>();
+			RHICmdList.SetLocalShaderUniformBuffer(GetPixelShader(), Parameter.GetBaseIndex(), LocalUB);
 		}
 	}
 
