Copying //UE4/Dev-Sequencer to //UE4/Dev-Main (Source: //UE4/Dev-Sequencer @ 3057646)

diff --git a/Engine/Source/Runtime/MovieSceneTracks/Public/Tracks/MovieSceneStringTrack.h b/Engine/Source/Runtime/MovieSceneTracks/Public/Tracks/MovieSceneStringTrack.h
new file mode 100644
--- /dev/null
+++ b/Engine/Source/Runtime/MovieSceneTracks/Public/Tracks/MovieSceneStringTrack.h
@@ -0,0 +1,68 @@
+// Copyright 1998-2016 Epic Games, Inc. All Rights Reserved.
+
+#pragma once
+
+#include "MovieScenePropertyTrack.h"
+#include "MovieSceneStringTrack.generated.h"
+
+
+/**
+ * Implements a movie scene track that holds a series of strings.
+ */
+UCLASS(MinimalAPI)
+class UMovieSceneStringTrack
+	: public UMovieScenePropertyTrack
+{
+	GENERATED_BODY()
+
+public:
+
+	/** Default constructor. */
+	UMovieSceneStringTrack()
+	{
+#if WITH_EDITORONLY_DATA
+		TrackTint = FColor(128, 128, 128);
+#endif
+	}
+
+public:
+
+	/**
+	 * Adds an string to the appropriate section.
+	 *
+	 * This method will create a new section if no appropriate one exists.
+	 *
+	 * @param Time The time at which the string should be triggered.
+	 * @param String The string to add.
+	 * @param KeyParams The keying parameters
+	 */
+	bool AddKeyToSection(float Time, const FString& String, FKeyParams KeyParams);
+
+	/**
+	 * Evaluates the track at the playback position
+	 *
+	 * @param Position The psootion at which to evaluate this track.
+	 * @param LastPositionThe last playback position.
+	 * @param OutString The string at the evaluation time.
+	 * @return true if anything was evaluated. Note: if false is returned OutColor remains unchanged
+	 */
+	virtual bool Eval(float Position, float LastPostion, FString& OutString) const;
+
+public:
+
+	//~ UMovieSceneTrack interface
+
+	virtual void AddSection(UMovieSceneSection& Section) override;
+	virtual TSharedPtr<IMovieSceneTrackInstance> CreateInstance() override;
+	virtual UMovieSceneSection* CreateNewSection() override;
+	virtual const TArray<UMovieSceneSection*>& GetAllSections() const override;
+	virtual TRange<float> GetSectionBoundaries() const override;
+	virtual bool HasSection(const UMovieSceneSection& Section) const override;
+	virtual bool IsEmpty() const override;
+	virtual void RemoveAllAnimationData() override;
+	virtual void RemoveSection(UMovieSceneSection& Section) override;
+
+#if WITH_EDITORONLY_DATA
+	virtual FText GetDefaultDisplayName() const override;
+#endif
+};
diff --git a/Engine/Source/Runtime/Slate/Private/Framework/Commands/UICommandList.cpp b/Engine/Source/Runtime/Slate/Private/Framework/Commands/UICommandList.cpp
--- a/Engine/Source/Runtime/Slate/Private/Framework/Commands/UICommandList.cpp
+++ b/Engine/Source/Runtime/Slate/Private/Framework/Commands/UICommandList.cpp
@@ -54,9 +54,7 @@ void FUICommandList::MapAction( const TSharedPtr< const FUICommandInfo > InUICom
 	check( InUICommandInfo.IsValid() );
 
 	// Check against already mapped actions
-#if DO_GUARD_SLOW
-	checkf(!UICommandBindingMap.Contains(InUICommandInfo), TEXT("Command list already contains a command named '%s'"), *InUICommandInfo->GetCommandName().ToString());
-#endif
+	checkfSlow(!UICommandBindingMap.Contains(InUICommandInfo), TEXT("Command list already contains a command named '%s'"), *InUICommandInfo->GetCommandName().ToString());
 
 	ContextsInList.Add( InUICommandInfo->GetBindingContext() );
 	UICommandBindingMap.Add( InUICommandInfo, InUIAction );
@@ -76,6 +74,13 @@ void FUICommandList::Append( const TSharedRef<FUICommandList>& InCommandsToAppen
 	InCommandsToAppend->ChildUICommandLists.AddUnique( this->AsShared() );
 }
 
+void FUICommandList::UnmapAction( const TSharedPtr< const FUICommandInfo > InUICommandInfo )
+{
+	// Check against already mapped actions
+	checkfSlow(UICommandBindingMap.Contains(InUICommandInfo), TEXT("Command list does not already contain a command named '%s'"), *InUICommandInfo->GetCommandName().ToString());
+
+	UICommandBindingMap.Remove( InUICommandInfo );
+}
 
 bool FUICommandList::ExecuteAction( const TSharedRef< const FUICommandInfo > InUICommandInfo ) const
 {
diff --git a/Engine/Source/Runtime/Slate/Public/Framework/Commands/UICommandList.h b/Engine/Source/Runtime/Slate/Public/Framework/Commands/UICommandList.h
--- a/Engine/Source/Runtime/Slate/Public/Framework/Commands/UICommandList.h
+++ b/Engine/Source/Runtime/Slate/Public/Framework/Commands/UICommandList.h
@@ -89,6 +89,13 @@ public:
 	 */
 	void Append( const TSharedRef<FUICommandList>& InCommandsToAppend );
 
+	/**
+	 * Unmaps a command info that was previously mapped
+	 *
+	 * @param InUICommandInfo		The command info to unmap
+	 */
+	void UnmapAction( const TSharedPtr< const FUICommandInfo > InUICommandInfo );
+
 	/**
 	 * Executes the action associated with the provided command info
 	 * Note: It is assumed at this point that CanExecuteAction was already checked
diff --git a/Engine/Source/Runtime/Slate/Public/Framework/Views/ITypedTableView.h b/Engine/Source/Runtime/Slate/Public/Framework/Views/ITypedTableView.h
--- a/Engine/Source/Runtime/Slate/Public/Framework/Views/ITypedTableView.h
+++ b/Engine/Source/Runtime/Slate/Public/Framework/Views/ITypedTableView.h
@@ -2,8 +2,14 @@
 
 #pragma once
 
+#include "SlateEnums.h"
 #include "ITypedTableView.generated.h"
 
+
+class ITableRow;
+struct FPointerEvent;
+
+
 UENUM()
 namespace ESelectionMode
 {
@@ -23,6 +29,7 @@ namespace ESelectionMode
 	};
 }
 
+
 /** Lists can exist as regular lists or as backing widgets for a tree. */
 UENUM()
 namespace ETableViewMode
@@ -35,6 +42,7 @@ namespace ETableViewMode
 	};
 }
 
+
 /** Interface for container widgets to talk to the owner list, tree, or grid. */
 template<typename ItemType>
 class ITypedTableView
diff --git a/Engine/Source/Runtime/UtilityShaders/Private/MediaShaders.cpp b/Engine/Source/Runtime/UtilityShaders/Private/MediaShaders.cpp
--- a/Engine/Source/Runtime/UtilityShaders/Private/MediaShaders.cpp
+++ b/Engine/Source/Runtime/UtilityShaders/Private/MediaShaders.cpp
@@ -4,6 +4,98 @@
 #include "MediaShaders.h"
 #include "ShaderParameterUtils.h"
 
+
+TGlobalResource<FMediaVertexDeclaration> GMediaVertexDeclaration;
+
+
+/* FMediaShadersVS shader
+ *****************************************************************************/
+
+IMPLEMENT_SHADER_TYPE(, FMediaShadersVS, TEXT("MediaShaders"), TEXT("MainVertexShader"), SF_Vertex);
+
+
+/* FAYUVConvertPS shader
+ *****************************************************************************/
+
+BEGIN_UNIFORM_BUFFER_STRUCT(FAYUVConvertUB, )
+DECLARE_UNIFORM_BUFFER_STRUCT_MEMBER_TEXTURE(Texture2D, AYUVTexture)
+DECLARE_UNIFORM_BUFFER_STRUCT_MEMBER_SAMPLER(SamplerState, AYUVSampler)
+END_UNIFORM_BUFFER_STRUCT(FAYUVConvertUB)
+
+IMPLEMENT_UNIFORM_BUFFER_STRUCT(FAYUVConvertUB, TEXT("AYUVConvertUB"));
+IMPLEMENT_SHADER_TYPE(, FAYUVConvertPS, TEXT("MediaShaders"), TEXT("AYUVConvertPS"), SF_Pixel);
+
+
+void FAYUVConvertPS::SetParameters(FRHICommandList& RHICmdList, TRefCountPtr<FRHITexture2D> AYUVTexture)
+{
+	FAYUVConvertUB UB;
+	{
+		UB.AYUVTexture = AYUVTexture;
+		UB.AYUVSampler = TStaticSamplerState<SF_Bilinear>::GetRHI();
+	}
+
+	TUniformBufferRef<FAYUVConvertUB> Data = TUniformBufferRef<FAYUVConvertUB>::CreateUniformBufferImmediate(UB, UniformBuffer_SingleFrame);
+	SetUniformBufferParameter(RHICmdList, GetPixelShader(), GetUniformBufferParameter<FAYUVConvertUB>(), Data);
+}
+
+
+/* FNV12ConvertPS shader
+ *****************************************************************************/
+
+BEGIN_UNIFORM_BUFFER_STRUCT(FNV12ConvertUB, )
+DECLARE_UNIFORM_BUFFER_STRUCT_MEMBER(uint32, Width)
+DECLARE_UNIFORM_BUFFER_STRUCT_MEMBER_TEXTURE(Texture2D, NV12Texture)
+DECLARE_UNIFORM_BUFFER_STRUCT_MEMBER_SAMPLER(SamplerState, NV12Sampler)
+END_UNIFORM_BUFFER_STRUCT(FNV12ConvertUB)
+
+IMPLEMENT_UNIFORM_BUFFER_STRUCT(FNV12ConvertUB, TEXT("NV12ConvertUB"));
+IMPLEMENT_SHADER_TYPE(, FNV12ConvertPS, TEXT("MediaShaders"), TEXT("NV12ConvertPS"), SF_Pixel);
+
+
+void FNV12ConvertPS::SetParameters(FRHICommandList& RHICmdList, TRefCountPtr<FRHITexture2D> NV12Texture)
+{
+	FNV12ConvertUB UB;
+	{
+		UB.Width = NV12Texture->GetSizeX();
+		UB.NV12Texture = NV12Texture;
+		UB.NV12Sampler = TStaticSamplerState<SF_Bilinear>::GetRHI();
+	}
+
+	TUniformBufferRef<FNV12ConvertUB> Data = TUniformBufferRef<FNV12ConvertUB>::CreateUniformBufferImmediate(UB, UniformBuffer_SingleFrame);
+	SetUniformBufferParameter(RHICmdList, GetPixelShader(), GetUniformBufferParameter<FNV12ConvertUB>(), Data);
+}
+
+
+/* FNV21ConvertPS shader
+ *****************************************************************************/
+
+BEGIN_UNIFORM_BUFFER_STRUCT(FNV21ConvertUB, )
+DECLARE_UNIFORM_BUFFER_STRUCT_MEMBER(uint32, Width)
+DECLARE_UNIFORM_BUFFER_STRUCT_MEMBER_TEXTURE(Texture2D, NV21Texture)
+DECLARE_UNIFORM_BUFFER_STRUCT_MEMBER_SAMPLER(SamplerState, NV21Sampler)
+END_UNIFORM_BUFFER_STRUCT(FNV21ConvertUB)
+
+IMPLEMENT_UNIFORM_BUFFER_STRUCT(FNV21ConvertUB, TEXT("NV21ConvertUB"));
+IMPLEMENT_SHADER_TYPE(, FNV21ConvertPS, TEXT("MediaShaders"), TEXT("NV21ConvertPS"), SF_Pixel);
+
+
+void FNV21ConvertPS::SetParameters(FRHICommandList& RHICmdList, TRefCountPtr<FRHITexture2D> NV21Texture)
+{
+	FNV21ConvertUB UB;
+	{
+		UB.Width = NV21Texture->GetSizeX();
+		UB.NV21Texture = NV21Texture;
+		UB.NV21Sampler = TStaticSamplerState<SF_Bilinear>::GetRHI();
+	}
+
+	TUniformBufferRef<FNV21ConvertUB> Data = TUniformBufferRef<FNV21ConvertUB>::CreateUniformBufferImmediate(UB, UniformBuffer_SingleFrame);
+	SetUniformBufferParameter(RHICmdList, GetPixelShader(), GetUniformBufferParameter<FNV21ConvertUB>(), Data);
+}
+
+
+/* FYCbCrConvertUB shader
+ *****************************************************************************/
+
 BEGIN_UNIFORM_BUFFER_STRUCT(FYCbCrConvertUB, )
 DECLARE_UNIFORM_BUFFER_STRUCT_MEMBER_TEXTURE(Texture2D, LumaTexture)
 DECLARE_UNIFORM_BUFFER_STRUCT_MEMBER_TEXTURE(Texture2D, CbCrTexture)
@@ -12,19 +104,133 @@ DECLARE_UNIFORM_BUFFER_STRUCT_MEMBER_SAMPLER(SamplerState, CbCrSampler)
 END_UNIFORM_BUFFER_STRUCT(FYCbCrConvertUB)
 
 IMPLEMENT_UNIFORM_BUFFER_STRUCT(FYCbCrConvertUB, TEXT("YCbCrConvertUB"));
+IMPLEMENT_SHADER_TYPE(, FYCbCrConvertPS, TEXT("MediaShaders"), TEXT("YCbCrConvertPS"), SF_Pixel);
 
-IMPLEMENT_SHADER_TYPE(, FMedaShadersVS, TEXT("MediaShaders"), TEXT("MainVertexShader"), SF_Vertex);
 
-IMPLEMENT_SHADER_TYPE(, FYCbCrConvertShaderPS, TEXT("MediaShaders"), TEXT("YCbCrConvertPS"), SF_Pixel);
-
-void FYCbCrConvertShaderPS::SetParameters(FRHICommandList& RHICmdList, FTextureRHIRef LumaTexture, FTextureRHIRef CbCrTexture)
+void FYCbCrConvertPS::SetParameters(FRHICommandList& RHICmdList, TRefCountPtr<FRHITexture2D> LumaTexture, TRefCountPtr<FRHITexture2D> CbCrTexture)
 {
-	FYCbCrConvertUB YCbCrUB;
-	YCbCrUB.LumaTexture = LumaTexture;
-	YCbCrUB.CbCrTexture = CbCrTexture;
-	YCbCrUB.LumaSampler = TStaticSamplerState<SF_Bilinear>::GetRHI();
-	YCbCrUB.CbCrSampler = TStaticSamplerState<SF_Bilinear>::GetRHI();
+	FYCbCrConvertUB UB;
+	{
+		UB.LumaTexture = LumaTexture;
+		UB.CbCrTexture = CbCrTexture;
+		UB.LumaSampler = TStaticSamplerState<SF_Bilinear>::GetRHI();
+		UB.CbCrSampler = TStaticSamplerState<SF_Bilinear>::GetRHI();
+	}
 
-	TUniformBufferRef<FYCbCrConvertUB> Data = TUniformBufferRef<FYCbCrConvertUB>::CreateUniformBufferImmediate(YCbCrUB, UniformBuffer_SingleFrame);
+	TUniformBufferRef<FYCbCrConvertUB> Data = TUniformBufferRef<FYCbCrConvertUB>::CreateUniformBufferImmediate(UB, UniformBuffer_SingleFrame);
 	SetUniformBufferParameter(RHICmdList, GetPixelShader(), GetUniformBufferParameter<FYCbCrConvertUB>(), Data);	
 }
+
+
+/* FUYVYConvertPS shader
+ *****************************************************************************/
+
+BEGIN_UNIFORM_BUFFER_STRUCT(FUYVYConvertUB, )
+DECLARE_UNIFORM_BUFFER_STRUCT_MEMBER(uint32, Width)
+DECLARE_UNIFORM_BUFFER_STRUCT_MEMBER_TEXTURE(Texture2D, UYVYTexture)
+DECLARE_UNIFORM_BUFFER_STRUCT_MEMBER_SAMPLER(SamplerState, UYVYSampler)
+END_UNIFORM_BUFFER_STRUCT(FUYVYConvertUB)
+
+IMPLEMENT_UNIFORM_BUFFER_STRUCT(FUYVYConvertUB, TEXT("UYVYConvertUB"));
+IMPLEMENT_SHADER_TYPE(, FUYVYConvertPS, TEXT("MediaShaders"), TEXT("UYVYConvertPS"), SF_Pixel);
+
+
+void FUYVYConvertPS::SetParameters(FRHICommandList& RHICmdList, TRefCountPtr<FRHITexture2D> UYVYTexture)
+{
+	FUYVYConvertUB UB;
+	{
+		UB.Width = UYVYTexture->GetSizeX();
+		UB.UYVYTexture = UYVYTexture;
+		UB.UYVYSampler = TStaticSamplerState<SF_Bilinear>::GetRHI();
+	}
+
+	TUniformBufferRef<FUYVYConvertUB> Data = TUniformBufferRef<FUYVYConvertUB>::CreateUniformBufferImmediate(UB, UniformBuffer_SingleFrame);
+	SetUniformBufferParameter(RHICmdList, GetPixelShader(), GetUniformBufferParameter<FUYVYConvertUB>(), Data);
+}
+
+
+/* FYUVConvertPS shader
+ *****************************************************************************/
+
+BEGIN_UNIFORM_BUFFER_STRUCT(FYUVConvertUB, )
+DECLARE_UNIFORM_BUFFER_STRUCT_MEMBER_TEXTURE(Texture2D, YTexture)
+DECLARE_UNIFORM_BUFFER_STRUCT_MEMBER_TEXTURE(Texture2D, UTexture)
+DECLARE_UNIFORM_BUFFER_STRUCT_MEMBER_TEXTURE(Texture2D, VTexture)
+DECLARE_UNIFORM_BUFFER_STRUCT_MEMBER_SAMPLER(SamplerState, YSampler)
+DECLARE_UNIFORM_BUFFER_STRUCT_MEMBER_SAMPLER(SamplerState, USampler)
+DECLARE_UNIFORM_BUFFER_STRUCT_MEMBER_SAMPLER(SamplerState, VSampler)
+END_UNIFORM_BUFFER_STRUCT(FYUVConvertUB)
+
+IMPLEMENT_UNIFORM_BUFFER_STRUCT(FYUVConvertUB, TEXT("YUVConvertUB"));
+IMPLEMENT_SHADER_TYPE(, FYUVConvertPS, TEXT("MediaShaders"), TEXT("YUVConvertPS"), SF_Pixel);
+
+
+void FYUVConvertPS::SetParameters(FRHICommandList& RHICmdList, TRefCountPtr<FRHITexture2D> YTexture, TRefCountPtr<FRHITexture2D> UTexture, TRefCountPtr<FRHITexture2D> VTexture)
+{
+	FYUVConvertUB UB;
+	{
+		UB.YTexture = YTexture;
+		UB.UTexture = UTexture;
+		UB.VTexture = VTexture;
+		UB.YSampler = TStaticSamplerState<SF_Bilinear>::GetRHI();
+		UB.USampler = TStaticSamplerState<SF_Bilinear>::GetRHI();
+		UB.VSampler = TStaticSamplerState<SF_Bilinear>::GetRHI();
+	}
+
+	TUniformBufferRef<FYUVConvertUB> Data = TUniformBufferRef<FYUVConvertUB>::CreateUniformBufferImmediate(UB, UniformBuffer_SingleFrame);
+	SetUniformBufferParameter(RHICmdList, GetPixelShader(), GetUniformBufferParameter<FYUVConvertUB>(), Data);
+}
+
+
+/* FYUY2ConvertPS shader
+ *****************************************************************************/
+
+BEGIN_UNIFORM_BUFFER_STRUCT(FYUY2ConvertUB, )
+DECLARE_UNIFORM_BUFFER_STRUCT_MEMBER(uint32, Width)
+DECLARE_UNIFORM_BUFFER_STRUCT_MEMBER_TEXTURE(Texture2D, YUY2Texture)
+DECLARE_UNIFORM_BUFFER_STRUCT_MEMBER_SAMPLER(SamplerState, YUY2Sampler)
+END_UNIFORM_BUFFER_STRUCT(FYUY2ConvertUB)
+
+IMPLEMENT_UNIFORM_BUFFER_STRUCT(FYUY2ConvertUB, TEXT("YUY2ConvertUB"));
+IMPLEMENT_SHADER_TYPE(, FYUY2ConvertPS, TEXT("MediaShaders"), TEXT("YUY2ConvertPS"), SF_Pixel);
+
+
+void FYUY2ConvertPS::SetParameters(FRHICommandList& RHICmdList, TRefCountPtr<FRHITexture2D> YUY2Texture)
+{
+	FYUY2ConvertUB UB;
+	{
+		UB.Width = YUY2Texture->GetSizeX();
+		UB.YUY2Texture = YUY2Texture;
+		UB.YUY2Sampler = TStaticSamplerState<SF_Bilinear>::GetRHI();
+	}
+
+	TUniformBufferRef<FYUY2ConvertUB> Data = TUniformBufferRef<FYUY2ConvertUB>::CreateUniformBufferImmediate(UB, UniformBuffer_SingleFrame);
+	SetUniformBufferParameter(RHICmdList, GetPixelShader(), GetUniformBufferParameter<FYUY2ConvertUB>(), Data);
+}
+
+
+/* FUYVYConvertPS shader
+ *****************************************************************************/
+
+BEGIN_UNIFORM_BUFFER_STRUCT(FYVYUConvertUB, )
+DECLARE_UNIFORM_BUFFER_STRUCT_MEMBER(uint32, Width)
+DECLARE_UNIFORM_BUFFER_STRUCT_MEMBER_TEXTURE(Texture2D, YVYUTexture)
+DECLARE_UNIFORM_BUFFER_STRUCT_MEMBER_SAMPLER(SamplerState, YVYUSampler)
+END_UNIFORM_BUFFER_STRUCT(FYVYUConvertUB)
+
+IMPLEMENT_UNIFORM_BUFFER_STRUCT(FYVYUConvertUB, TEXT("YVYUConvertUB"));
+IMPLEMENT_SHADER_TYPE(, FYVYUConvertPS, TEXT("MediaShaders"), TEXT("YVYUConvertPS"), SF_Pixel);
+
+
+void FYVYUConvertPS::SetParameters(FRHICommandList& RHICmdList, TRefCountPtr<FRHITexture2D> YVYUTexture)
+{
+	FYVYUConvertUB UB;
+	{
+		UB.Width = YVYUTexture->GetSizeX();
+		UB.YVYUTexture = YVYUTexture;
+		UB.YVYUSampler = TStaticSamplerState<SF_Bilinear>::GetRHI();
+	}
+
+	TUniformBufferRef<FYVYUConvertUB> Data = TUniformBufferRef<FYVYUConvertUB>::CreateUniformBufferImmediate(UB, UniformBuffer_SingleFrame);
+	SetUniformBufferParameter(RHICmdList, GetPixelShader(), GetUniformBufferParameter<FYVYUConvertUB>(), Data);
+}
diff --git a/Engine/Source/Runtime/UtilityShaders/Public/MediaShaders.h b/Engine/Source/Runtime/UtilityShaders/Public/MediaShaders.h
--- a/Engine/Source/Runtime/UtilityShaders/Public/MediaShaders.h
+++ b/Engine/Source/Runtime/UtilityShaders/Public/MediaShaders.h
@@ -1,28 +1,67 @@
 // Copyright 1998-2016 Epic Games, Inc. All Rights Reserved.
 
-/*================================================================================
-MediaShaders.h: Shaders for converting decoded media files to RGBA format
-================================================================================*/
+#pragma once
 
-#include "Engine/EngineTypes.h"
-#include "RenderCore.h"
+#include "Engine/EngineTypes.h"  // @todo: fix GlobalShader.h include dependencies
 #include "GlobalShader.h"
-#include "ShaderParameters.h"
-#include "ShaderParameterUtils.h"
-#include "RHIStaticStates.h"
+#include "RenderResource.h"
 
-#pragma once
 
-class FMediaShaderHelper
+/**
+ * Stores media drawing vertices.
+ */
+struct FMediaElementVertex
 {
-	//draws a simple full screen quad with 0-1 UVs.
-	//user must set up shaders and global shader state before calling.
-	static void DrawFullScreenQuad();
+	FVector4 Position;
+	FVector2D TextureCoordinate;
+
+	FMediaElementVertex() { }
+
+	FMediaElementVertex(const FVector4& InPosition, const FVector2D& InTextureCoordinate)
+		: Position(InPosition)
+		, TextureCoordinate(InTextureCoordinate)
+	{ }
 };
 
-class FMedaShadersVS : public FGlobalShader
+
+/**
+ * The simple element vertex declaration resource type.
+ */
+class FMediaVertexDeclaration
+	: public FRenderResource
 {
-	DECLARE_SHADER_TYPE(FMedaShadersVS, Global);
+public:
+
+	FVertexDeclarationRHIRef VertexDeclarationRHI;
+
+	virtual ~FMediaVertexDeclaration() { }
+
+	virtual void InitRHI() override
+	{
+		FVertexDeclarationElementList Elements;
+		uint16 Stride = sizeof(FMediaElementVertex);
+		Elements.Add(FVertexElement(0, STRUCT_OFFSET(FMediaElementVertex, Position), VET_Float4, 0, Stride));
+		Elements.Add(FVertexElement(0, STRUCT_OFFSET(FMediaElementVertex, TextureCoordinate), VET_Float2, 1, Stride));
+		VertexDeclarationRHI = RHICreateVertexDeclaration(Elements);
+	}
+
+	virtual void ReleaseRHI() override
+	{
+		VertexDeclarationRHI.SafeRelease();
+	}
+};
+
+
+UTILITYSHADERS_API extern TGlobalResource<FMediaVertexDeclaration> GMediaVertexDeclaration;
+
+
+/**
+ * Media vertex shader (shared by all media shaders).
+ */
+class FMediaShadersVS
+	: public FGlobalShader
+{
+	DECLARE_EXPORTED_SHADER_TYPE(FMediaShadersVS, Global, UTILITYSHADERS_API);
 
 public:
 
@@ -32,27 +71,67 @@ public:
 	}
 
 	/** Default constructor. */
-	FMedaShadersVS() {}
+	FMediaShadersVS() { }
 
 	/** Initialization constructor. */
-	FMedaShadersVS(const ShaderMetaType::CompiledShaderInitializerType& Initializer)
+	FMediaShadersVS(const ShaderMetaType::CompiledShaderInitializerType& Initializer)
 		: FGlobalShader(Initializer)
+	{ }
+
+	virtual bool Serialize(FArchive& Ar) override
 	{
+		bool bShaderHasOutdatedParameters = FGlobalShader::Serialize(Ar);
+		return bShaderHasOutdatedParameters;
 	}
+};
+
+
+/**
+ * Pixel shader to convert an AYUV texture to RGBA.
+ *
+ * This shader expects a single texture consisting of a N x M array of pixels
+ * in AYUV format. Each pixel is encoded as four consecutive unsigned chars
+ * with the following layout: [V0 U0 Y0 A0][V1 U1 Y1 A1]..
+ */
+class FAYUVConvertPS
+	: public FGlobalShader
+{
+	DECLARE_EXPORTED_SHADER_TYPE(FAYUVConvertPS, Global, UTILITYSHADERS_API);
+
+public:
+
+	static bool ShouldCache(EShaderPlatform Platform)
+	{
+		return IsFeatureLevelSupported(Platform, ERHIFeatureLevel::SM4);
+	}
+
+	FAYUVConvertPS() { }
+
+	FAYUVConvertPS(const ShaderMetaType::CompiledShaderInitializerType& Initializer)
+		: FGlobalShader(Initializer)
+	{ }
 
 	virtual bool Serialize(FArchive& Ar) override
 	{
 		bool bShaderHasOutdatedParameters = FGlobalShader::Serialize(Ar);
 		return bShaderHasOutdatedParameters;
-	}	
+	}
 
-private:
+	UTILITYSHADERS_API void SetParameters(FRHICommandList& RHICmdList, TRefCountPtr<FRHITexture2D> AYUVTexture);
 };
 
-/** Shader to convert a full-res luma (Y) and half-res chroma (CbCr) texture to RGBA */
-class FYCbCrConvertShaderPS : public FGlobalShader
+
+/**
+ * Pixel shader to convert a NV12 frame to RGBA.
+ *
+ * This shader expects an NV12 frame packed into single texture in PF_G8 format.
+ *
+ * @see http://www.fourcc.org/yuv.php#NV12
+ */
+class FNV12ConvertPS
+	: public FGlobalShader
 {
-	DECLARE_SHADER_TYPE(FYCbCrConvertShaderPS, Global);
+	DECLARE_EXPORTED_SHADER_TYPE(FNV12ConvertPS, Global, UTILITYSHADERS_API);
 
 public:
 
@@ -61,24 +140,231 @@ public:
 		return IsFeatureLevelSupported(Platform, ERHIFeatureLevel::SM4);
 	}
 
-	/** Default constructor. */
-	FYCbCrConvertShaderPS() {}
+	FNV12ConvertPS() { }
 
-	/** Initialization constructor. */
-	FYCbCrConvertShaderPS(const ShaderMetaType::CompiledShaderInitializerType& Initializer)
+	FNV12ConvertPS(const ShaderMetaType::CompiledShaderInitializerType& Initializer)
 		: FGlobalShader(Initializer)
-	{		
+	{ }
+
+	virtual bool Serialize(FArchive& Ar) override
+	{
+		bool bShaderHasOutdatedParameters = FGlobalShader::Serialize(Ar);
+		return bShaderHasOutdatedParameters;
 	}
 
+	UTILITYSHADERS_API void SetParameters(FRHICommandList& RHICmdList, TRefCountPtr<FRHITexture2D> NV12Texture);
+};
+
+
+/**
+ * Pixel shader to convert a NV21 frame to RGBA.
+ *
+ * This shader expects an NV21 frame packed into single texture in PF_G8 format.
+ *
+ * @see http://www.fourcc.org/yuv.php#NV21
+ */
+class FNV21ConvertPS
+	: public FGlobalShader
+{
+	DECLARE_EXPORTED_SHADER_TYPE(FNV21ConvertPS, Global, UTILITYSHADERS_API);
+
+public:
+
+	static bool ShouldCache(EShaderPlatform Platform)
+	{
+		return IsFeatureLevelSupported(Platform, ERHIFeatureLevel::SM4);
+	}
+
+	FNV21ConvertPS() { }
+
+	FNV21ConvertPS(const ShaderMetaType::CompiledShaderInitializerType& Initializer)
+		: FGlobalShader(Initializer)
+	{ }
+
+	virtual bool Serialize(FArchive& Ar) override
+	{
+		bool bShaderHasOutdatedParameters = FGlobalShader::Serialize(Ar);
+		return bShaderHasOutdatedParameters;
+	}
+
+	UTILITYSHADERS_API void SetParameters(FRHICommandList& RHICmdList, TRefCountPtr<FRHITexture2D> NV21Texture);
+};
+
+
+/**
+ * Pixel shader to convert a PS4 YCbCr texture to RGBA.
+ *
+ * This shader expects a separate chroma and luma plane stored in two textures
+ * in PF_B8G8R8A8 format. The full-size luma plane contains the Y-components.
+ * The half-size chroma plane contains the UV components in the following
+ * memory layout: [U0, V0][U1, V1]
+ * 
+ */
+class FYCbCrConvertPS
+	: public FGlobalShader
+{
+	DECLARE_EXPORTED_SHADER_TYPE(FYCbCrConvertPS, Global, UTILITYSHADERS_API);
+
+public:
+
+	static bool ShouldCache(EShaderPlatform Platform)
+	{
+		return IsFeatureLevelSupported(Platform, ERHIFeatureLevel::SM4);
+	}
+
+	FYCbCrConvertPS() { }
+
+	FYCbCrConvertPS(const ShaderMetaType::CompiledShaderInitializerType& Initializer)
+		: FGlobalShader(Initializer)
+	{ }
+
 	virtual bool Serialize(FArchive& Ar) override
 	{
 		bool bShaderHasOutdatedParameters = FGlobalShader::Serialize(Ar);		
 		return bShaderHasOutdatedParameters;
 	}
 
-	void SetParameters(FRHICommandList& RHICmdList, FTextureRHIRef LumaTexture, FTextureRHIRef CbCrTexture);
+	UTILITYSHADERS_API void SetParameters(FRHICommandList& RHICmdList, TRefCountPtr<FRHITexture2D> LumaTexture, TRefCountPtr<FRHITexture2D> CbCrTexture);
+};
+
+
+/**
+ * Pixel shader to convert a UYVY (Y422, UYNV) frame to RGBA.
+ *
+ * This shader expects a UYVY frame packed into a single texture in PF_B8G8R8A8
+ * format with the following memory layout: [U0, Y0, V1, Y1][U1, Y2, V1, Y3]..
+ *
+ * @see http://www.fourcc.org/yuv.php#UYVY
+ */
+class FUYVYConvertPS
+	: public FGlobalShader
+{
+	DECLARE_EXPORTED_SHADER_TYPE(FUYVYConvertPS, Global, UTILITYSHADERS_API);
+
+public:
+
+	static bool ShouldCache(EShaderPlatform Platform)
+	{
+		return IsFeatureLevelSupported(Platform, ERHIFeatureLevel::SM4);
+	}
+
+	FUYVYConvertPS() { }
+
+	FUYVYConvertPS(const ShaderMetaType::CompiledShaderInitializerType& Initializer)
+		: FGlobalShader(Initializer)
+	{ }
+
+	virtual bool Serialize(FArchive& Ar) override
+	{
+		bool bShaderHasOutdatedParameters = FGlobalShader::Serialize(Ar);
+		return bShaderHasOutdatedParameters;
+	}
+
+	UTILITYSHADERS_API void SetParameters(FRHICommandList& RHICmdList, TRefCountPtr<FRHITexture2D> UYVYTexture);
+};
+
+
+/**
+ * Pixel shader to convert Y, U, and V planes to RGBA.
+ *
+ * This shader expects three textures in PF_G8 format,
+ * one for each plane of Y, U, and V components.
+ */
+class FYUVConvertPS
+	: public FGlobalShader
+{
+	DECLARE_EXPORTED_SHADER_TYPE(FYUVConvertPS, Global, UTILITYSHADERS_API);
+
+public:
+
+	static bool ShouldCache(EShaderPlatform Platform)
+	{
+		return IsFeatureLevelSupported(Platform, ERHIFeatureLevel::SM4);
+	}
 
-private:
+	FYUVConvertPS() { }
+
+	FYUVConvertPS(const ShaderMetaType::CompiledShaderInitializerType& Initializer)
+		: FGlobalShader(Initializer)
+	{ }
+
+	virtual bool Serialize(FArchive& Ar) override
+	{
+		bool bShaderHasOutdatedParameters = FGlobalShader::Serialize(Ar);
+		return bShaderHasOutdatedParameters;
+	}
+
+	UTILITYSHADERS_API void SetParameters(FRHICommandList& RHICmdList, TRefCountPtr<FRHITexture2D> YTexture, TRefCountPtr<FRHITexture2D> UTexture, TRefCountPtr<FRHITexture2D> VTexture);
+};
 
-	
-};
\ No newline at end of file
+
+/**
+ * Pixel shader to convert a YUY2 frame to RGBA.
+ *
+ * This shader expects an YUY2 frame packed into single texture in PF_B8G8R8A8
+ * format with the following memory layout: [Y0, U0, Y1, V0][Y2, U1, Y3, V1]...
+ *
+ * @see http://www.fourcc.org/yuv.php#YUY2
+ */
+class FYUY2ConvertPS
+	: public FGlobalShader
+{
+	DECLARE_EXPORTED_SHADER_TYPE(FYUY2ConvertPS, Global, UTILITYSHADERS_API);
+
+public:
+
+	static bool ShouldCache(EShaderPlatform Platform)
+	{
+		return IsFeatureLevelSupported(Platform, ERHIFeatureLevel::SM4);
+	}
+
+	FYUY2ConvertPS() { }
+
+	FYUY2ConvertPS(const ShaderMetaType::CompiledShaderInitializerType& Initializer)
+		: FGlobalShader(Initializer)
+	{ }
+
+	virtual bool Serialize(FArchive& Ar) override
+	{
+		bool bShaderHasOutdatedParameters = FGlobalShader::Serialize(Ar);
+		return bShaderHasOutdatedParameters;
+	}
+
+	UTILITYSHADERS_API void SetParameters(FRHICommandList& RHICmdList, TRefCountPtr<FRHITexture2D> YUY2Texture);
+};
+
+
+/**
+ * Pixel shader to convert a YVYU frame to RGBA.
+ *
+ * This shader expects a YVYU frame packed into a single texture in PF_B8G8R8A8
+ * format with the following memory layout: [Y0, V0, Y1, U0][Y2, V1, Y3, U1]..
+ *
+ * @see http://www.fourcc.org/yuv.php#YVYU
+ */
+class FYVYUConvertPS
+	: public FGlobalShader
+{
+	DECLARE_EXPORTED_SHADER_TYPE(FYVYUConvertPS, Global, UTILITYSHADERS_API);
+
+public:
+
+	static bool ShouldCache(EShaderPlatform Platform)
+	{
+		return IsFeatureLevelSupported(Platform, ERHIFeatureLevel::SM4);
+	}
+
+	FYVYUConvertPS() { }
+
+	FYVYUConvertPS(const ShaderMetaType::CompiledShaderInitializerType& Initializer)
+		: FGlobalShader(Initializer)
+	{ }
+
+	virtual bool Serialize(FArchive& Ar) override
+	{
+		bool bShaderHasOutdatedParameters = FGlobalShader::Serialize(Ar);
+		return bShaderHasOutdatedParameters;
+	}
+
+	UTILITYSHADERS_API void SetParameters(FRHICommandList& RHICmdList, TRefCountPtr<FRHITexture2D> YVYUTexture);
+};
diff --git a/Engine/Source/Runtime/UtilityShaders/UtilityShaders.Build.cs b/Engine/Source/Runtime/UtilityShaders/UtilityShaders.Build.cs
--- a/Engine/Source/Runtime/UtilityShaders/UtilityShaders.Build.cs
+++ b/Engine/Source/Runtime/UtilityShaders/UtilityShaders.Build.cs
@@ -6,7 +6,7 @@ public class UtilityShaders : ModuleRules
 {
 	public UtilityShaders(TargetInfo Target)
 	{
-		PrivateDependencyModuleNames.AddRange(
+		PublicDependencyModuleNames.AddRange(
 			new string[] {
 				"Core",
 				"Engine",
