UE-43767 //UE4/Main: Step 'Build DDC Win64' - 3 Errors
-Added runtime check that morpheus plugin is available to compiling the rgbatoyuv shader.  Fixed the name of the shader in the 'optional' shader list. Wrapped shader in HAS_MORPHEUS.  All of this makes it more like PostProcessMorpheus, which works correctly on all build configurations.
-Cleaned up some comments.

diff --git a/Engine/Shaders/MorpheusInclude.usf b/Engine/Shaders/MorpheusInclude.usf
--- a/Engine/Shaders/MorpheusInclude.usf
+++ b/Engine/Shaders/MorpheusInclude.usf
@@ -3,6 +3,4 @@
 #include "Common.usf"
 #include "PostProcessCommon.usf"
 #include "PS4/PostProcessHMDMorpheus.usf"
-#if PS4_PROFILE
-#include "PS4/RGBAToYUV420.usf"
-#endif
\ No newline at end of file
+#include "PS4/RGBAToYUV420.usf"
\ No newline at end of file
diff --git a/Engine/Source/Runtime/ShaderCore/Private/ShaderCore.cpp b/Engine/Source/Runtime/ShaderCore/Private/ShaderCore.cpp
--- a/Engine/Source/Runtime/ShaderCore/Private/ShaderCore.cpp
+++ b/Engine/Source/Runtime/ShaderCore/Private/ShaderCore.cpp
@@ -442,7 +442,7 @@ void GetShaderIncludes(const TCHAR* Filename, TArray<FString>& IncludeFilenames,
 					const bool bIsOptionalInclude = (ExtractedIncludeFilename == TEXT("PS4/PS4Common.usf") 
 						|| ExtractedIncludeFilename == TEXT("PS4/PostProcessHMDMorpheus.usf")
 						|| ExtractedIncludeFilename == TEXT("PS4/RTWriteMaskProcessing.usf")
-						|| ExtractedIncludeFilename == TEXT("PS4/RGBAToYUV.usf")
+						|| ExtractedIncludeFilename == TEXT("PS4/RGBAToYUV420.usf")
 						|| ExtractedIncludeFilename == TEXT("XboxOne/XboxOneCommon.usf")
 						);
 					// ignore the header if it's optional and doesn't exist
diff --git a/Engine/Source/Runtime/UtilityShaders/Private/RGBAToYUV420Shader.cpp b/Engine/Source/Runtime/UtilityShaders/Private/RGBAToYUV420Shader.cpp
--- a/Engine/Source/Runtime/UtilityShaders/Private/RGBAToYUV420Shader.cpp
+++ b/Engine/Source/Runtime/UtilityShaders/Private/RGBAToYUV420Shader.cpp
@@ -3,6 +3,8 @@
 #include "RGBAToYUV420Shader.h"
 #include "ShaderParameterUtils.h"
 
+#if defined(HAS_MORPHEUS) && HAS_MORPHEUS
+
 /* FRGBAToYUV420CS shader
  *****************************************************************************/
 
@@ -44,3 +46,5 @@ void FRGBAToYUV420CS::UnbindBuffers(FRHICommandList& RHICmdList)
 	FComputeShaderRHIParamRef ComputeShaderRHI = GetComputeShader();
 	RHICmdList.SetUAVParameter(ComputeShaderRHI, OutTextureRW.GetBaseIndex(), FUnorderedAccessViewRHIParamRef());
 }
+
+#endif //defined(HAS_MORPHEUS) && HAS_MORPHEUS
diff --git a/Engine/Source/Runtime/UtilityShaders/Public/RGBAToYUV420Shader.h b/Engine/Source/Runtime/UtilityShaders/Public/RGBAToYUV420Shader.h
--- a/Engine/Source/Runtime/UtilityShaders/Public/RGBAToYUV420Shader.h
+++ b/Engine/Source/Runtime/UtilityShaders/Public/RGBAToYUV420Shader.h
@@ -7,26 +7,35 @@
 #include "RenderResource.h"
 #include "Shader.h"
 #include "GlobalShader.h"
+#include "Modules/ModuleManager.h"
+
+#if defined(HAS_MORPHEUS) && HAS_MORPHEUS
 
 /**
- * Pixel shader to convert an AYUV texture to RGBA.
+ * Pixel shader to convert RGBA to YUV420.
+ *
+ * YUV420 has 8 bit intensity values in the top 2/3 of the texture for every pixel and 8 bit each UV coordinates 
+ * into the YUV colorspace in the bottom 1/3 for every pixel quad.
  *
- * This shader expects a single texture consisting of a N x M array of pixels
- * in AYUV format. Each pixel is encoded as four consecutive unsigned chars
- * with the following layout: [V0 U0 Y0 A0][V1 U1 Y1 A1]..
+ * This is only used by PS4.
  */
 class FRGBAToYUV420CS
 	: public FGlobalShader
 {
-	//DECLARE_EXPORTED_SHADER_TYPE(FRGBAToYUV420CS, Global, UTILITYSHADERS_API);
 	DECLARE_SHADER_TYPE(FRGBAToYUV420CS, Global);
 
 public:
 
 	static bool ShouldCache(EShaderPlatform Platform)
 	{
-		//return true;  
-		return Platform == EShaderPlatform::SP_PS4;
+		if (Platform == EShaderPlatform::SP_PS4)
+		{
+			// we must use a run time check for this because the builds the build machines create will have Morpheus defined,
+			// but a user will not necessarily have the Morpheus files
+			return FModuleManager::Get().ModuleExists(TEXT("Morpheus"));
+		}
+
+		return false;
 	}
 
 	FRGBAToYUV420CS() { }
@@ -49,3 +58,5 @@ public:
 protected:
 	FShaderResourceParameter OutTextureRW;
 };
+
+#endif //defined(HAS_MORPHEUS) && HAS_MORPHEUS
