Copying //UE4/Dev-Rendering to //UE4/Dev-Main (Source: //UE4/Dev-Rendering @ 3461187)

diff --git a/Engine/Source/Runtime/ShaderCore/Public/ShaderCache.h b/Engine/Source/Runtime/ShaderCore/Public/ShaderCache.h
--- a/Engine/Source/Runtime/ShaderCore/Public/ShaderCache.h
+++ b/Engine/Source/Runtime/ShaderCore/Public/ShaderCache.h
@@ -709,10 +709,10 @@ private:
 	bool bBatchingPaused;
 	
 	//Thread access mutual exclusion
-	mutable FRWMutex ShaderCacheGlobalStateMutex;
-	mutable FRWMutex PipelineStateMutex;
-	mutable FRWMutex DrawLogMutex;
-	mutable FRWMutex ContextCacheStatesMutex;
+	mutable FRWLock ShaderCacheGlobalStateMutex;
+	mutable FRWLock PipelineStateMutex;
+	mutable FRWLock DrawLogMutex;
+	mutable FRWLock ContextCacheStatesMutex;
 	
 	//List of states per RHI context with a default state
 	FShaderCacheState* DefaultCacheState;
diff --git a/Engine/Source/Runtime/ShaderCore/Public/ShaderCacheTypes.h b/Engine/Source/Runtime/ShaderCore/Public/ShaderCacheTypes.h
--- a/Engine/Source/Runtime/ShaderCore/Public/ShaderCacheTypes.h
+++ b/Engine/Source/Runtime/ShaderCore/Public/ShaderCacheTypes.h
@@ -10,187 +10,6 @@ ShaderCacheTypes.h: Shader Cache Types
 #include "RHI.h"
 #include "RHIDefinitions.h"
 
-/**
- * FRWMutex - Read/Write Mutex
- *	- Provides Read/Write (or shared-exclusive) access.  
- *	- As a fallback default for non implemented platforms, using a single FCriticalSection to provide complete single mutual exclusion - no seperate Read/Write access.
- *
- *	TODOs:
- *		- Add Xbox and Windows specific lock structures/calls Ref: https://msdn.microsoft.com/en-us/library/windows/desktop/aa904937(v=vs.85).aspx
- *	Notes:
- *		- pThreads and Win32 API's don't provide a mechanism for upgrading a ownership of a read lock to a write lock - to get round that this system unlocks then acquires a write lock.
- *		- If trying to use system described here (https://en.wikipedia.org/wiki/Readers–writer_lock) to implement a read-write lock:
- *			- The two mutex solution relies on a seperate thread being able to unlock a mutex locked on a different thread (Needs more work than described to implement if this is not the case).
- *			- Doesn't seem to be a generic 'engine' conditional variable for the mutex+conditional solution.
- *		- Currently has an experimental thread_local lock level implementation (set DEVELOPER_TEST_RWLOCK_LEVELS to 1) for potential deadlock detection (see: http://www.drdobbs.com/parallel/use-lock-hierarchies-to-avoid-deadlock/204801163?pgno=4)
- *			- Define the locking order by passing an integer in the Mutex constructor
- *			- To avoid deadlocks lock in increasing sequence order (can miss ones out but can't go lower)
- *			- Asserts if trying to get a same or lower level lock as this could result in deadlock depending on situation/timing etc.
- */
- 
-#define DEVELOPER_TEST_RWLOCK_LEVELS 0
-#define RWLOCK_DEBUG_CHECKS (DO_CHECK && PLATFORM_DESKTOP && !UE_BUILD_SHIPPING && !UE_BUILD_SHIPPING_WITH_EDITOR && DEVELOPER_TEST_RWLOCK_LEVELS)
- 
-class FRWMutex
-{
-	friend class FRWScopeLock;
-	
-public:
-
-	// Default to level zero, we are not allowed to take another at the same level or lower - so no more than one 'default' mutex at once
-	FRWMutex(uint32 Level = 0)
-#if RWLOCK_DEBUG_CHECKS
-	: MutexLevel(Level)
-#endif
-	{
-#if PLATFORM_USE_PTHREADS
-		int Err = pthread_rwlock_init(&Mutex, nullptr);
-		checkf(Err == 0, TEXT("pthread_rwlock_init failed with error: %d"), Err);
-#endif
-	}
-
-	~FRWMutex()
-	{
-#if PLATFORM_USE_PTHREADS
-		int Err = pthread_rwlock_destroy(&Mutex);
-		checkf(Err == 0, TEXT("pthread_rwlock_destroy failed with error: %d"), Err);
-#endif
-	}
-	
-private:
-
-#if RWLOCK_DEBUG_CHECKS
-	void CheckAndAddLockLevel()
-	{
-		if(LockLevels.Num())
-		{
-			uint32 Last = LockLevels.Last();
-			check(MutexLevel > Last && "Trying to lock on an object with the same or lower level, options: [change lock order|change lock level values|disable this]");
-		}
-		
-		LockLevels.Add(MutexLevel);
-	}
-#endif
-	
-	void ReadLock()
-	{
-#if RWLOCK_DEBUG_CHECKS
-		CheckAndAddLockLevel();
-#endif
-	
-#if PLATFORM_USE_PTHREADS
-		int Err = pthread_rwlock_rdlock(&Mutex);
-		checkf(Err == 0, TEXT("pthread_rwlock_rdlock failed with error: %d"), Err);
-#else
-		Mutex.Lock();
-#endif
-	}
-	
-	void WriteLock()
-	{
-#if RWLOCK_DEBUG_CHECKS
-		CheckAndAddLockLevel();
-#endif
-
-#if PLATFORM_USE_PTHREADS
-		int Err = pthread_rwlock_wrlock(&Mutex);
-		checkf(Err == 0, TEXT("pthread_rwlock_wrlock failed with error: %d"), Err);
-#else
-		Mutex.Lock();
-#endif
-	}
-	
-	void Unlock()
-	{
-#if PLATFORM_USE_PTHREADS
-		int Err = pthread_rwlock_unlock(&Mutex);
-		checkf(Err == 0, TEXT("pthread_rwlock_unlock failed with error: %d"), Err);
-#else
-		Mutex.Unlock();
-#endif
-
-#if RWLOCK_DEBUG_CHECKS
-		LockLevels.Pop();
-#endif
-	}
-	
-	void RaiseLockToWrite()
-	{
-#if PLATFORM_USE_PTHREADS
-		int Err = pthread_rwlock_unlock(&Mutex);
-		checkf(Err == 0, TEXT("pthread_rwlock_unlock failed with error: %d"), Err);
-		Err = pthread_rwlock_wrlock(&Mutex);
-		checkf(Err == 0, TEXT("pthread_rwlock_wrlock failed with error: %d"), Err);
-#else
-		//NOP
-#endif		
-	}
-
-private:
-#if PLATFORM_USE_PTHREADS
-	pthread_rwlock_t Mutex;
-#else
-	FCriticalSection Mutex;
-#endif
-
-#if RWLOCK_DEBUG_CHECKS
-	uint32 MutexLevel;
-	static thread_local TArray<uint32> LockLevels;
-#endif
-};
-
-//
-// A scope lifetime controlled Read or Write lock of referenced mutex object
-//
-enum FRWScopeLockType
-{
-	SLT_ReadOnly = 0,
-	SLT_Write,
-};
-
-class FRWScopeLock
-{
-public:
-	FRWScopeLock(FRWMutex& InLockObject,FRWScopeLockType InLockType) 
-	: LockObject(InLockObject)
-	, LockType(InLockType)
-	{
-		if(LockType != SLT_ReadOnly)
-		{
-			LockObject.WriteLock();
-		}
-		else
-		{
-			LockObject.ReadLock();
-		}
-	}
-	
-	void RaiseLockToWrite()
-	{	
-		if(LockType == SLT_ReadOnly)
-		{
-			LockObject.RaiseLockToWrite();
-			LockType = SLT_Write;
-		}
-	}
-
-	~FRWScopeLock()
-	{
-		LockObject.Unlock();
-	}
-
-private:
-
-	FRWScopeLock();
-	FRWScopeLock(const FRWScopeLock&);
-	FRWScopeLock& operator=(const FRWScopeLock&);
-	
-private:
-	FRWMutex& LockObject;
-	FRWScopeLockType LockType;
-};
-
-
 /** Texture type enum for shader cache draw keys */
 enum EShaderCacheTextureType
 {
diff --git a/Engine/Source/Runtime/ShaderCore/Public/ShaderCore.h b/Engine/Source/Runtime/ShaderCore/Public/ShaderCore.h
--- a/Engine/Source/Runtime/ShaderCore/Public/ShaderCore.h
+++ b/Engine/Source/Runtime/ShaderCore/Public/ShaderCore.h
@@ -897,6 +897,11 @@ extern SHADERCORE_API const class FSHAHash& GetShaderFileHash(const TCHAR* Filen
  */
 extern SHADERCORE_API const class FSHAHash& GetShaderFilesHash(const TArray<FString>& Filenames);
 
+/**
+* Find the shader in the potential shader directories (engine, project, plugins) relative to the engine shader folder
+*/
+extern SHADERCORE_API FString FindShaderRelativePath(const FString Filename);
+
 extern void BuildShaderFileToUniformBufferMap(TMap<FString, TArray<const TCHAR*> >& ShaderFileToUniformBufferVariables);
 
 /**
diff --git a/Engine/Source/Runtime/ShaderCore/Public/ShaderParameterUtils.h b/Engine/Source/Runtime/ShaderCore/Public/ShaderParameterUtils.h
--- a/Engine/Source/Runtime/ShaderCore/Public/ShaderParameterUtils.h
+++ b/Engine/Source/Runtime/ShaderCore/Public/ShaderParameterUtils.h
@@ -350,60 +350,70 @@ TRHICmdList& RHICmdList,
 	}
 }
 
-template<typename TShaderRHIRef, typename TRHICmdList>
-inline void SetUAVParameterIfCS(TRHICmdList& RHICmdList, TShaderRHIRef, const FShaderResourceParameter& UAVParameter, FUnorderedAccessViewRHIParamRef UAV)
-{
-}
 
-template<> inline void SetUAVParameterIfCS<FComputeShaderRHIParamRef, FRHICommandList>(FRHICommandList& RHICmdList, FComputeShaderRHIParamRef ComputeShader,const FShaderResourceParameter& UAVParameter,FUnorderedAccessViewRHIParamRef UAV)
+template<typename TRHICmdList>
+inline bool SetUAVParameterIfCS(TRHICmdList& RHICmdList, const FVertexShaderRHIParamRef Shader, const FShaderResourceParameter& UAVParameter, FUnorderedAccessViewRHIParamRef UAV)
 {
-	SetUAVParameter(RHICmdList, ComputeShader,UAVParameter,UAV);
+	return false;
 }
 
-template<> inline void SetUAVParameterIfCS<FComputeShaderRHIRef, FRHICommandList>(FRHICommandList& RHICmdList, FComputeShaderRHIRef ComputeShader, const FShaderResourceParameter& UAVParameter, FUnorderedAccessViewRHIParamRef UAV)
+template<typename TRHICmdList>
+inline bool SetUAVParameterIfCS(TRHICmdList& RHICmdList, const FPixelShaderRHIParamRef Shader, const FShaderResourceParameter& UAVParameter, FUnorderedAccessViewRHIParamRef UAV)
 {
-	SetUAVParameter(RHICmdList, ComputeShader,UAVParameter,UAV);
+	return false;
 }
 
-template<> inline void SetUAVParameterIfCS<FComputeShaderRHIParamRef, FRHICommandListImmediate>(FRHICommandListImmediate& RHICmdList, FComputeShaderRHIParamRef ComputeShader,const FShaderResourceParameter& UAVParameter,FUnorderedAccessViewRHIParamRef UAV)
+template<typename TRHICmdList>
+inline bool SetUAVParameterIfCS(TRHICmdList& RHICmdList, const FHullShaderRHIParamRef Shader, const FShaderResourceParameter& UAVParameter, FUnorderedAccessViewRHIParamRef UAV)
 {
-	SetUAVParameter(RHICmdList, ComputeShader,UAVParameter,UAV);
+	return false;
 }
 
-template<> inline void SetUAVParameterIfCS<FComputeShaderRHIRef, FRHICommandListImmediate>(FRHICommandListImmediate& RHICmdList, FComputeShaderRHIRef ComputeShader, const FShaderResourceParameter& UAVParameter, FUnorderedAccessViewRHIParamRef UAV)
+template<typename TRHICmdList>
+inline bool SetUAVParameterIfCS(TRHICmdList& RHICmdList, const FDomainShaderRHIParamRef Shader, const FShaderResourceParameter& UAVParameter, FUnorderedAccessViewRHIParamRef UAV)
 {
-	SetUAVParameter(RHICmdList, ComputeShader,UAVParameter,UAV);
+	return false;
 }
 
-template<> inline void SetUAVParameterIfCS<FComputeShaderRHIParamRef, FRHIAsyncComputeCommandListImmediate>(FRHIAsyncComputeCommandListImmediate& RHICmdList, FComputeShaderRHIParamRef ComputeShader, const FShaderResourceParameter& UAVParameter, FUnorderedAccessViewRHIParamRef UAV)
+template<typename TRHICmdList>
+inline bool SetUAVParameterIfCS(TRHICmdList& RHICmdList, const FGeometryShaderRHIParamRef Shader, const FShaderResourceParameter& UAVParameter, FUnorderedAccessViewRHIParamRef UAV)
 {
-	SetUAVParameter(RHICmdList, ComputeShader, UAVParameter, UAV);
+	return false;
 }
 
-template<> inline void SetUAVParameterIfCS<FComputeShaderRHIRef, FRHIAsyncComputeCommandListImmediate>(FRHIAsyncComputeCommandListImmediate& RHICmdList, FComputeShaderRHIRef ComputeShader, const FShaderResourceParameter& UAVParameter, FUnorderedAccessViewRHIParamRef UAV)
+template<typename TRHICmdList>
+inline bool SetUAVParameterIfCS(TRHICmdList& RHICmdList, const FComputeShaderRHIParamRef Shader, const FShaderResourceParameter& UAVParameter, FUnorderedAccessViewRHIParamRef UAV)
 {
-	SetUAVParameter(RHICmdList, ComputeShader, UAVParameter, UAV);
+	SetUAVParameter(RHICmdList, Shader, UAVParameter, UAV);
+	return UAVParameter.IsBound();
 }
 
+
 template<typename TShaderRHIRef, typename TRHICmdList>
 inline void FRWShaderParameter::SetBuffer(TRHICmdList& RHICmdList, TShaderRHIRef Shader, const FRWBuffer& RWBuffer) const
 {
-	SetSRVParameter(RHICmdList, Shader,SRVParameter,RWBuffer.SRV);
-	SetUAVParameterIfCS(RHICmdList, Shader,UAVParameter,RWBuffer.UAV);
+	if (!SetUAVParameterIfCS(RHICmdList, Shader, UAVParameter, RWBuffer.UAV))
+	{
+		SetSRVParameter(RHICmdList, Shader, SRVParameter, RWBuffer.SRV);
+	}
 }
 
 template<typename TShaderRHIRef, typename TRHICmdList>
 inline void FRWShaderParameter::SetBuffer(TRHICmdList& RHICmdList, TShaderRHIRef Shader, const FRWBufferStructured& RWBuffer) const
 {
-	SetSRVParameter(RHICmdList, Shader,SRVParameter,RWBuffer.SRV);
-	SetUAVParameterIfCS(RHICmdList, Shader,UAVParameter,RWBuffer.UAV);
+	if (!SetUAVParameterIfCS(RHICmdList, Shader, UAVParameter, RWBuffer.UAV))
+	{
+		SetSRVParameter(RHICmdList, Shader, SRVParameter, RWBuffer.SRV);
+	}
 }
 
 template<typename TShaderRHIRef, typename TRHICmdList>
 inline void FRWShaderParameter::SetTexture(TRHICmdList& RHICmdList, TShaderRHIRef Shader, const FTextureRHIParamRef Texture, FUnorderedAccessViewRHIParamRef UAV) const
 {
-	SetTextureParameter(RHICmdList, Shader,SRVParameter,Texture);
-	SetUAVParameterIfCS(RHICmdList, Shader,UAVParameter,UAV);
+	if (!SetUAVParameterIfCS(RHICmdList, Shader, UAVParameter, UAV))
+	{
+		SetTextureParameter(RHICmdList, Shader, SRVParameter, Texture);
+	}
 }
 
 template<typename TRHICmdList>
diff --git a/Engine/Source/Runtime/ShaderCore/Public/VertexFactory.h b/Engine/Source/Runtime/ShaderCore/Public/VertexFactory.h
--- a/Engine/Source/Runtime/ShaderCore/Public/VertexFactory.h
+++ b/Engine/Source/Runtime/ShaderCore/Public/VertexFactory.h
@@ -413,12 +413,6 @@ public:
 	 */
 	virtual uint64 GetStaticBatchElementVisibility(const class FSceneView& View, const struct FMeshBatch* Batch) const { return 1; }
 
-	/**
-	  * Get a bitmask representing the visibility of each FMeshBatch element for the shadow depths of a given light.
-	  * FMeshBatch.bRequiresPerElementVisibility must be set for this to be called.
-	  */
-	virtual uint64 GetStaticBatchElementShadowVisibility(const class FSceneView& View, const class FLightSceneProxy* LightSceneProxy, const struct FMeshBatch* Batch) const { return GetStaticBatchElementVisibility(View, Batch); }
-
 protected:
 
 	/**
diff --git a/Engine/Source/Runtime/SlateCore/Private/Rendering/DrawElements.cpp b/Engine/Source/Runtime/SlateCore/Private/Rendering/DrawElements.cpp
--- a/Engine/Source/Runtime/SlateCore/Private/Rendering/DrawElements.cpp
+++ b/Engine/Source/Runtime/SlateCore/Private/Rendering/DrawElements.cpp
@@ -21,7 +21,7 @@ void FSlateDataPayload::SetTextPayloadProperties( FSlateWindowElementList& Eleme
 	int32 StartIndex = FMath::Min<int32>(InStartIndex, Count - 1);
 	int32 EndIndex = FMath::Min<int32>(InEndIndex, Count - 1);
 	Count = 1 + ((EndIndex > StartIndex) ? EndIndex - StartIndex : 0);
-	ImmutableText = (TCHAR*)ElementList.Alloc(sizeof(TCHAR) * Count, ALIGNOF(TCHAR));
+	ImmutableText = (TCHAR*)ElementList.Alloc(sizeof(TCHAR) * Count, alignof(TCHAR));
 	if (Count > 1)
 	{
 		FCString::Strncpy(ImmutableText, InText.GetCharArray().GetData() + StartIndex, Count);
diff --git a/Engine/Source/Runtime/SlateCore/Public/Rendering/DrawElements.h b/Engine/Source/Runtime/SlateCore/Public/Rendering/DrawElements.h
--- a/Engine/Source/Runtime/SlateCore/Public/Rendering/DrawElements.h
+++ b/Engine/Source/Runtime/SlateCore/Public/Rendering/DrawElements.h
@@ -1333,7 +1333,7 @@ public:
 	template <typename T>
 	FORCEINLINE_DEBUGGABLE void* Alloc()
 	{
-		return MemManager.Alloc(sizeof(T), ALIGNOF(T));
+		return MemManager.Alloc(sizeof(T), alignof(T));
 	}
 
 	FSlateBatchData& GetBatchData() { return BatchData; }
diff --git a/Engine/Source/Runtime/SlateRHIRenderer/Private/SlateRHIRenderer.cpp b/Engine/Source/Runtime/SlateRHIRenderer/Private/SlateRHIRenderer.cpp
--- a/Engine/Source/Runtime/SlateRHIRenderer/Private/SlateRHIRenderer.cpp
+++ b/Engine/Source/Runtime/SlateRHIRenderer/Private/SlateRHIRenderer.cpp
@@ -576,6 +576,8 @@ void FSlateRHIRenderer::DrawWindow_RenderThread(FRHICommandListImmediate& RHICmd
 	// Should only be called by the rendering thread
 	check(IsInRenderingThread());
 
+	FMaterialRenderProxy::UpdateDeferredCachedUniformExpressions();
+
 	// Optional off-screen UI composition during HDR rendering
 	static const auto CVarCompositeMode = IConsoleManager::Get().FindTConsoleVariableDataInt(TEXT("r.HDR.UI.CompositeMode"));
 	static const auto CVarHDROutputEnabled = IConsoleManager::Get().FindTConsoleVariableDataInt(TEXT("r.HDR.EnableHDROutput"));
@@ -691,7 +693,7 @@ void FSlateRHIRenderer::DrawWindow_RenderThread(FRHICommandListImmediate& RHICmd
 		}
 
 #if DEBUG_OVERDRAW
-		DrawClearQuad(RHICmdList, GMaxRHIFeatureLevel, FLinearColor::Black);
+		DrawClearQuad(RHICmdList, FLinearColor::Black);
 #endif
 		if( BatchData.GetRenderBatches().Num() > 0 )
 		{
diff --git a/Engine/Source/Runtime/SlateRHIRenderer/Private/SlateRHIRenderingPolicy.cpp b/Engine/Source/Runtime/SlateRHIRenderer/Private/SlateRHIRenderingPolicy.cpp
--- a/Engine/Source/Runtime/SlateRHIRenderer/Private/SlateRHIRenderingPolicy.cpp
+++ b/Engine/Source/Runtime/SlateRHIRenderer/Private/SlateRHIRenderingPolicy.cpp
@@ -165,7 +165,7 @@ void FSlateRHIRenderingPolicy::UpdateVertexAndIndexBuffers(FRHICommandListImmedi
 		VertexBuffer.PreFillBuffer(NumVertices, bShouldShrinkResources);
 		IndexBuffer.PreFillBuffer(NumIndices, bShouldShrinkResources);
 
-		if(!GRHIThread || RHICmdList.Bypass())
+		if(!IsRunningRHIInSeparateThread() || RHICmdList.Bypass())
 		{
 			uint8* VertexBufferData = (uint8*)VertexBuffer.LockBuffer_RenderThread(NumVertices);
 			uint8* IndexBufferData =  (uint8*)IndexBuffer.LockBuffer_RenderThread(NumIndices);
diff --git a/Engine/Source/Runtime/SlateRHIRenderer/Private/SlateRHIResourceManager.cpp b/Engine/Source/Runtime/SlateRHIRenderer/Private/SlateRHIResourceManager.cpp
--- a/Engine/Source/Runtime/SlateRHIRenderer/Private/SlateRHIResourceManager.cpp
+++ b/Engine/Source/Runtime/SlateRHIRenderer/Private/SlateRHIResourceManager.cpp
@@ -267,7 +267,7 @@ bool FSlateRHIResourceManager::IsAtlasPageResourceAlphaOnly() const
 void FSlateRHIResourceManager::Tick(float DeltaSeconds)
 {
 	// Don't need to do this if there's no RHI thread.
-	if ( GRHIThread )
+	if (IsRunningRHIInSeparateThread())
 	{
 		struct FDeleteCachedRenderDataContext
 		{
@@ -1058,17 +1058,17 @@ void FSlateRHIResourceManager::ReleaseCachedBuffer(FRHICommandListImmediate& RHI
 {
 	check(IsInRenderingThread());
 
-	if ( GRHIThread )
+	if (IsRunningRHIInSeparateThread())
 	{
 		PooledBuffersPendingRelease.Add(PooledBuffer);
 		PooledBuffer->ReleaseResourcesFence = RHICmdList.RHIThreadFence();
 	}
 	else
 	{
-			PooledBuffer->VertexBuffer.Destroy();
-			PooledBuffer->IndexBuffer.Destroy();
-			delete PooledBuffer;
-		}
+		PooledBuffer->VertexBuffer.Destroy();
+		PooledBuffer->IndexBuffer.Destroy();
+		delete PooledBuffer;
+	}
 }
 
 void FSlateRHIResourceManager::ReleaseResources()
diff --git a/Engine/Source/Runtime/SlateRHIRenderer/Private/SlateUpdatableBuffer.cpp b/Engine/Source/Runtime/SlateRHIRenderer/Private/SlateUpdatableBuffer.cpp
--- a/Engine/Source/Runtime/SlateRHIRenderer/Private/SlateUpdatableBuffer.cpp
+++ b/Engine/Source/Runtime/SlateRHIRenderer/Private/SlateUpdatableBuffer.cpp
@@ -90,7 +90,7 @@ void FSlateUpdatableInstanceBuffer::UpdateRenderingData_RenderThread(FRHICommand
 	const TArray<FVector4>& RenderThreadBufferData = BufferData[BufferIndex];
 	InstanceBufferResource.PreFillBuffer( RenderThreadBufferData.Num(), false );
 
-	if(!GRHIThread || RHICmdList.Bypass())
+	if(!IsRunningRHIInSeparateThread() || RHICmdList.Bypass())
 	{
 		uint8* InstanceBufferData = (uint8*)InstanceBufferResource.LockBuffer_RenderThread(RenderThreadBufferData.Num());
 
diff --git a/Engine/Source/Runtime/UtilityShaders/Private/ClearQuad.cpp b/Engine/Source/Runtime/UtilityShaders/Private/ClearQuad.cpp
--- a/Engine/Source/Runtime/UtilityShaders/Private/ClearQuad.cpp
+++ b/Engine/Source/Runtime/UtilityShaders/Private/ClearQuad.cpp
@@ -8,7 +8,7 @@
 #include "ClearReplacementShaders.h"
 #include "RendererInterface.h"
 
-static void ClearQuadSetup( FRHICommandList& RHICmdList, ERHIFeatureLevel::Type FeatureLevel, bool bClearColor, int32 NumClearColors, const FLinearColor* ClearColorArray, bool bClearDepth, float Depth, bool bClearStencil, uint32 Stencil )
+static void ClearQuadSetup( FRHICommandList& RHICmdList, bool bClearColor, int32 NumClearColors, const FLinearColor* ClearColorArray, bool bClearDepth, float Depth, bool bClearStencil, uint32 Stencil )
 {
 	// Set new states
 	FBlendStateRHIParamRef BlendStateRHI;
@@ -43,7 +43,7 @@ static void ClearQuadSetup( FRHICommandList& RHICmdList, ERHIFeatureLevel::Type
 	GraphicsPSOInit.BlendState = BlendStateRHI;
 	GraphicsPSOInit.DepthStencilState = DepthStencilStateRHI;
 
-	auto ShaderMap = GetGlobalShaderMap(FeatureLevel);
+	auto ShaderMap = GetGlobalShaderMap(GMaxRHIFeatureLevel);
 
 
 	// Set the new shaders
@@ -106,9 +106,9 @@ static void ClearQuadSetup( FRHICommandList& RHICmdList, ERHIFeatureLevel::Type
 }
 
 const uint32 GMaxSizeUAVDMA = 1024;
-static void ClearUAVShader(FRHICommandList& RHICmdList, ERHIFeatureLevel::Type FeatureLevel, FUnorderedAccessViewRHIParamRef UnorderedAccessViewRHI, uint32 Size, uint32 Value)
+static void ClearUAVShader(FRHICommandList& RHICmdList, FUnorderedAccessViewRHIParamRef UnorderedAccessViewRHI, uint32 Size, uint32 Value)
 {
-	TShaderMapRef<FClearBufferReplacementCS> ComputeShader(GetGlobalShaderMap(FeatureLevel));
+	TShaderMapRef<FClearBufferReplacementCS> ComputeShader(GetGlobalShaderMap(GMaxRHIFeatureLevel));
 	FComputeShaderRHIParamRef ShaderRHI = ComputeShader->GetComputeShader();
 	RHICmdList.SetComputeShader(ShaderRHI);
 	ComputeShader->SetParameters(RHICmdList, UnorderedAccessViewRHI, Value);
@@ -118,7 +118,7 @@ static void ClearUAVShader(FRHICommandList& RHICmdList, ERHIFeatureLevel::Type F
 	ComputeShader->FinalizeParameters(RHICmdList, UnorderedAccessViewRHI);
 }
 
-void ClearUAV(FRHICommandList& RHICmdList, ERHIFeatureLevel::Type FeatureLevel, const FRWBufferStructured& StructuredBuffer, uint32 Value)
+void ClearUAV(FRHICommandList& RHICmdList, const FRWBufferStructured& StructuredBuffer, uint32 Value)
 {
 	if (StructuredBuffer.NumBytes <= GMaxSizeUAVDMA)
 	{
@@ -127,11 +127,11 @@ void ClearUAV(FRHICommandList& RHICmdList, ERHIFeatureLevel::Type FeatureLevel,
 	}
 	else
 	{
-		ClearUAVShader(RHICmdList, FeatureLevel, StructuredBuffer.UAV, StructuredBuffer.NumBytes, Value);
+		ClearUAVShader(RHICmdList, StructuredBuffer.UAV, StructuredBuffer.NumBytes, Value);
 	}
 }
 
-void ClearUAV(FRHICommandList& RHICmdList, ERHIFeatureLevel::Type FeatureLevel, const FRWBuffer& Buffer, uint32 Value)
+void ClearUAV(FRHICommandList& RHICmdList, const FRWBuffer& Buffer, uint32 Value)
 {
 	if (Buffer.NumBytes <= GMaxSizeUAVDMA)
 	{
@@ -140,16 +140,29 @@ void ClearUAV(FRHICommandList& RHICmdList, ERHIFeatureLevel::Type FeatureLevel,
 	}
 	else
 	{
-		ClearUAVShader(RHICmdList, FeatureLevel, Buffer.UAV, Buffer.NumBytes, Value);
+		ClearUAVShader(RHICmdList, Buffer.UAV, Buffer.NumBytes, Value);
+	}
+}
+
+void ClearUAV(FRHICommandList& RHICmdList, FRHIUnorderedAccessView* Buffer, uint32 NumBytes, uint32 Value)
+{
+	if (NumBytes <= GMaxSizeUAVDMA)
+	{
+		uint32 Values[4] = { Value, Value, Value, Value };
+		RHICmdList.ClearTinyUAV(Buffer, Values);
+	}
+	else
+	{
+		ClearUAVShader(RHICmdList, Buffer, NumBytes, Value);
 	}
 }
 
 template< typename T >
-inline void ClearUAV_T(FRHICommandList& RHICmdList, ERHIFeatureLevel::Type FeatureLevel, const FSceneRenderTargetItem& RenderTargetItem, const T(&ClearValues)[4])
+inline void ClearUAV_T(FRHICommandList& RHICmdList, const FSceneRenderTargetItem& RenderTargetItem, const T(&ClearValues)[4])
 {
 	if (auto Texture2d = RenderTargetItem.TargetableTexture->GetTexture2D())
 	{
-		TShaderMapRef< FClearTexture2DReplacementCS<T> > ComputeShader(GetGlobalShaderMap(FeatureLevel));
+		TShaderMapRef< FClearTexture2DReplacementCS<T> > ComputeShader(GetGlobalShaderMap(GMaxRHIFeatureLevel));
 		FComputeShaderRHIParamRef ShaderRHI = ComputeShader->GetComputeShader();
 		RHICmdList.SetComputeShader(ShaderRHI);
 		ComputeShader->SetParameters(RHICmdList, RenderTargetItem.UAV, ClearValues);
@@ -160,7 +173,7 @@ inline void ClearUAV_T(FRHICommandList& RHICmdList, ERHIFeatureLevel::Type Featu
 	}
 	else if (auto Texture2dArray = RenderTargetItem.TargetableTexture->GetTexture2DArray())
 	{
-		TShaderMapRef< FClearTexture2DArrayReplacementCS<T> > ComputeShader(GetGlobalShaderMap(FeatureLevel));
+		TShaderMapRef< FClearTexture2DArrayReplacementCS<T> > ComputeShader(GetGlobalShaderMap(GMaxRHIFeatureLevel));
 		FComputeShaderRHIParamRef ShaderRHI = ComputeShader->GetComputeShader();
 		RHICmdList.SetComputeShader(ShaderRHI);
 		ComputeShader->SetParameters(RHICmdList, RenderTargetItem.UAV, ClearValues);
@@ -172,7 +185,7 @@ inline void ClearUAV_T(FRHICommandList& RHICmdList, ERHIFeatureLevel::Type Featu
 	}
 	else if (auto TextureCube = RenderTargetItem.TargetableTexture->GetTextureCube())
 	{
-		TShaderMapRef< FClearTexture2DArrayReplacementCS<T> > ComputeShader(GetGlobalShaderMap(FeatureLevel));
+		TShaderMapRef< FClearTexture2DArrayReplacementCS<T> > ComputeShader(GetGlobalShaderMap(GMaxRHIFeatureLevel));
 		FComputeShaderRHIParamRef ShaderRHI = ComputeShader->GetComputeShader();
 		RHICmdList.SetComputeShader(ShaderRHI);
 		ComputeShader->SetParameters(RHICmdList, RenderTargetItem.UAV, ClearValues);
@@ -183,7 +196,7 @@ inline void ClearUAV_T(FRHICommandList& RHICmdList, ERHIFeatureLevel::Type Featu
 	}
 	else if (auto Texture3d = RenderTargetItem.TargetableTexture->GetTexture3D())
 	{
-		TShaderMapRef< FClearVolumeReplacementCS<T> > ComputeShader(GetGlobalShaderMap(FeatureLevel));
+		TShaderMapRef< FClearVolumeReplacementCS<T> > ComputeShader(GetGlobalShaderMap(GMaxRHIFeatureLevel));
 		FComputeShaderRHIParamRef ShaderRHI = ComputeShader->GetComputeShader();
 		RHICmdList.SetComputeShader(ShaderRHI);
 		ComputeShader->SetParameters(RHICmdList, RenderTargetItem.UAV, ClearValues);
@@ -199,9 +212,9 @@ inline void ClearUAV_T(FRHICommandList& RHICmdList, ERHIFeatureLevel::Type Featu
 	}
 }
 
-void ClearTexture2DUAV(FRHICommandList& RHICmdList, ERHIFeatureLevel::Type FeatureLevel, FUnorderedAccessViewRHIParamRef UAV, int32 Width, int32 Height, const FLinearColor& ClearColor)
+void ClearTexture2DUAV(FRHICommandList& RHICmdList, FUnorderedAccessViewRHIParamRef UAV, int32 Width, int32 Height, const FLinearColor& ClearColor)
 {
-	TShaderMapRef< FClearTexture2DReplacementCS<float> > ComputeShader(GetGlobalShaderMap(FeatureLevel));
+	TShaderMapRef< FClearTexture2DReplacementCS<float> > ComputeShader(GetGlobalShaderMap(GMaxRHIFeatureLevel));
 	FComputeShaderRHIParamRef ShaderRHI = ComputeShader->GetComputeShader();
 	RHICmdList.SetComputeShader(ShaderRHI);
 	ComputeShader->SetParameters(RHICmdList, UAV, reinterpret_cast<const float(&)[4]>(ClearColor));
@@ -211,24 +224,24 @@ void ClearTexture2DUAV(FRHICommandList& RHICmdList, ERHIFeatureLevel::Type Featu
 	ComputeShader->FinalizeParameters(RHICmdList, UAV);
 }
 
-void ClearUAV(FRHICommandList& RHICmdList, ERHIFeatureLevel::Type FeatureLevel, const FSceneRenderTargetItem& RenderTargetItem, const float(&ClearValues)[4])
+void ClearUAV(FRHICommandList& RHICmdList, const FSceneRenderTargetItem& RenderTargetItem, const float(&ClearValues)[4])
 {
-	ClearUAV_T(RHICmdList, FeatureLevel, RenderTargetItem, ClearValues);
+	ClearUAV_T(RHICmdList, RenderTargetItem, ClearValues);
 }
 
-void ClearUAV(FRHICommandList& RHICmdList, ERHIFeatureLevel::Type FeatureLevel, const FSceneRenderTargetItem& RenderTargetItem, const uint32(&ClearValues)[4])
+void ClearUAV(FRHICommandList& RHICmdList, const FSceneRenderTargetItem& RenderTargetItem, const uint32(&ClearValues)[4])
 {
-	ClearUAV_T(RHICmdList, FeatureLevel, RenderTargetItem, ClearValues);
+	ClearUAV_T(RHICmdList, RenderTargetItem, ClearValues);
 }
 
-void ClearUAV(FRHICommandList& RHICmdList, ERHIFeatureLevel::Type FeatureLevel, const FSceneRenderTargetItem& RenderTargetItem, const FLinearColor& ClearColor)
+void ClearUAV(FRHICommandList& RHICmdList, const FSceneRenderTargetItem& RenderTargetItem, const FLinearColor& ClearColor)
 {
-	ClearUAV_T(RHICmdList, FeatureLevel, RenderTargetItem, reinterpret_cast<const float(&)[4]>(ClearColor));
+	ClearUAV_T(RHICmdList, RenderTargetItem, reinterpret_cast<const float(&)[4]>(ClearColor));
 }
 
-void DrawClearQuadMRT(FRHICommandList& RHICmdList, ERHIFeatureLevel::Type FeatureLevel, bool bClearColor, int32 NumClearColors, const FLinearColor* ClearColorArray, bool bClearDepth, float Depth, bool bClearStencil, uint32 Stencil)
+void DrawClearQuadMRT(FRHICommandList& RHICmdList, bool bClearColor, int32 NumClearColors, const FLinearColor* ClearColorArray, bool bClearDepth, float Depth, bool bClearStencil, uint32 Stencil)
 {
-	ClearQuadSetup(RHICmdList, FeatureLevel, bClearColor, NumClearColors, ClearColorArray, bClearDepth, Depth, bClearStencil, Stencil);
+	ClearQuadSetup(RHICmdList, bClearColor, NumClearColors, ClearColorArray, bClearDepth, Depth, bClearStencil, Stencil);
 
 	// without a hole
 	FVector4 Vertices[4];
@@ -239,7 +252,7 @@ void DrawClearQuadMRT(FRHICommandList& RHICmdList, ERHIFeatureLevel::Type Featur
 	DrawPrimitiveUP(RHICmdList, PT_TriangleStrip, 2, Vertices, sizeof(Vertices[0]));
 }
 
-void DrawClearQuadMRT(FRHICommandList& RHICmdList, ERHIFeatureLevel::Type FeatureLevel, bool bClearColor, int32 NumClearColors, const FLinearColor* ClearColorArray, bool bClearDepth, float Depth, bool bClearStencil, uint32 Stencil, FIntPoint ViewSize, FIntRect ExcludeRect)
+void DrawClearQuadMRT(FRHICommandList& RHICmdList, bool bClearColor, int32 NumClearColors, const FLinearColor* ClearColorArray, bool bClearDepth, float Depth, bool bClearStencil, uint32 Stencil, FIntPoint ViewSize, FIntRect ExcludeRect)
 {
 	if (ExcludeRect.Min == FIntPoint::ZeroValue && ExcludeRect.Max == ViewSize)
 	{
@@ -247,7 +260,7 @@ void DrawClearQuadMRT(FRHICommandList& RHICmdList, ERHIFeatureLevel::Type Featur
 		return;
 	}
 
-	ClearQuadSetup(RHICmdList, FeatureLevel, bClearColor, NumClearColors, ClearColorArray, bClearDepth, Depth, bClearStencil, Stencil);
+	ClearQuadSetup(RHICmdList, bClearColor, NumClearColors, ClearColorArray, bClearDepth, Depth, bClearStencil, Stencil);
 
 	// Draw a fullscreen quad
 	if (ExcludeRect.Width() > 0 && ExcludeRect.Height() > 0)
@@ -294,4 +307,4 @@ void DrawClearQuadMRT(FRHICommandList& RHICmdList, ERHIFeatureLevel::Type Featur
 
 		DrawPrimitiveUP(RHICmdList, PT_TriangleStrip, 2, Vertices, sizeof(Vertices[0]));
 	}
-}
\ No newline at end of file
+}
diff --git a/Engine/Source/Runtime/UtilityShaders/Private/ClearReplacementShaders.cpp b/Engine/Source/Runtime/UtilityShaders/Private/ClearReplacementShaders.cpp
--- a/Engine/Source/Runtime/UtilityShaders/Private/ClearReplacementShaders.cpp
+++ b/Engine/Source/Runtime/UtilityShaders/Private/ClearReplacementShaders.cpp
@@ -18,7 +18,7 @@ void FClearTexture2DReplacementCS<T>::SetParameters( FRHICommandList& RHICmdList
 template< typename T >
 void FClearTexture2DReplacementCS<T>::FinalizeParameters(FRHICommandList& RHICmdList, FUnorderedAccessViewRHIParamRef TextureRW)
 {
-	RHICmdList.TransitionResource(EResourceTransitionAccess::EReadable, EResourceTransitionPipeline::EComputeToGfx, TextureRW);
+	RHICmdList.TransitionResource(EResourceTransitionAccess::ERWBarrier, EResourceTransitionPipeline::EComputeToCompute, TextureRW);
 }
 
 template< typename T >
@@ -34,7 +34,7 @@ void FClearTexture2DArrayReplacementCS<T>::SetParameters(FRHICommandList& RHICmd
 template< typename T >
 void FClearTexture2DArrayReplacementCS<T>::FinalizeParameters(FRHICommandList& RHICmdList, FUnorderedAccessViewRHIParamRef TextureRW)
 {
-	RHICmdList.TransitionResource(EResourceTransitionAccess::EReadable, EResourceTransitionPipeline::EComputeToGfx, TextureRW);
+	RHICmdList.TransitionResource(EResourceTransitionAccess::ERWBarrier, EResourceTransitionPipeline::EComputeToCompute, TextureRW);
 }
 
 template< typename T >
@@ -50,7 +50,7 @@ void FClearVolumeReplacementCS<T>::SetParameters(FRHICommandList& RHICmdList, FU
 template< typename T >
 void FClearVolumeReplacementCS<T>::FinalizeParameters(FRHICommandList& RHICmdList, FUnorderedAccessViewRHIParamRef TextureRW)
 {
-	RHICmdList.TransitionResource(EResourceTransitionAccess::EReadable, EResourceTransitionPipeline::EComputeToGfx, TextureRW);
+	RHICmdList.TransitionResource(EResourceTransitionAccess::ERWBarrier, EResourceTransitionPipeline::EComputeToCompute, TextureRW);
 }
 
 void FClearTexture2DReplacementScissorCS::SetParameters(FRHICommandList& RHICmdList, FUnorderedAccessViewRHIParamRef TextureRW, FLinearColor InClearColor, const FVector4& InTargetBounds)
@@ -64,7 +64,7 @@ void FClearTexture2DReplacementScissorCS::SetParameters(FRHICommandList& RHICmdL
 
 void FClearTexture2DReplacementScissorCS::FinalizeParameters(FRHICommandList& RHICmdList, FUnorderedAccessViewRHIParamRef TextureRW)
 {
-	RHICmdList.TransitionResource(EResourceTransitionAccess::EReadable, EResourceTransitionPipeline::EComputeToGfx, TextureRW);
+	RHICmdList.TransitionResource(EResourceTransitionAccess::ERWBarrier, EResourceTransitionPipeline::EComputeToCompute, TextureRW);
 }
 
 void FClearBufferReplacementCS::SetParameters( FRHICommandList& RHICmdList, FUnorderedAccessViewRHIParamRef BufferRW, uint32 Dword )
@@ -77,7 +77,7 @@ void FClearBufferReplacementCS::SetParameters( FRHICommandList& RHICmdList, FUno
 
 void FClearBufferReplacementCS::FinalizeParameters(FRHICommandList& RHICmdList, FUnorderedAccessViewRHIParamRef BufferRW)
 {
-	RHICmdList.TransitionResource(EResourceTransitionAccess::EReadable, EResourceTransitionPipeline::EComputeToGfx, BufferRW);
+	RHICmdList.TransitionResource(EResourceTransitionAccess::ERWBarrier, EResourceTransitionPipeline::EComputeToCompute, BufferRW);
 }
 
 template class FClearTexture2DReplacementCS<float>;
diff --git a/Engine/Source/Runtime/UtilityShaders/Public/ClearQuad.h b/Engine/Source/Runtime/UtilityShaders/Public/ClearQuad.h
--- a/Engine/Source/Runtime/UtilityShaders/Public/ClearQuad.h
+++ b/Engine/Source/Runtime/UtilityShaders/Public/ClearQuad.h
@@ -12,30 +12,31 @@ struct FSceneRenderTargetItem;
 class FRHIUnorderedAccessView;
 
 extern UTILITYSHADERS_API const uint32 GMaxSizeUAVDMA;
-extern UTILITYSHADERS_API void ClearUAV(FRHICommandList& RHICmdList, ERHIFeatureLevel::Type FeatureLevel, const FRWBufferStructured& StructuredBuffer, uint32 Value);
-extern UTILITYSHADERS_API void ClearUAV(FRHICommandList& RHICmdList, ERHIFeatureLevel::Type FeatureLevel, const FRWBuffer& Buffer, uint32 Value);
+extern UTILITYSHADERS_API void ClearUAV(FRHICommandList& RHICmdList, const FRWBufferStructured& StructuredBuffer, uint32 Value);
+extern UTILITYSHADERS_API void ClearUAV(FRHICommandList& RHICmdList, const FRWBuffer& Buffer, uint32 Value);
+extern UTILITYSHADERS_API void ClearUAV(FRHICommandList& RHICmdList, FRHIUnorderedAccessView* Buffer, uint32 NumBytes, uint32 Value);
 
-extern UTILITYSHADERS_API void ClearUAV(FRHICommandList& RHICmdList, ERHIFeatureLevel::Type FeatureLevel, const FSceneRenderTargetItem& RenderTargetItem, const float(&ClearValues)[4]);
-extern UTILITYSHADERS_API void ClearUAV(FRHICommandList& RHICmdList, ERHIFeatureLevel::Type FeatureLevel, const FSceneRenderTargetItem& RenderTargetItem, const uint32(&ClearValues)[4]);
-extern UTILITYSHADERS_API void ClearUAV(FRHICommandList& RHICmdList, ERHIFeatureLevel::Type FeatureLevel, const FSceneRenderTargetItem& RenderTargetItem,  const FLinearColor& ClearColor);
+extern UTILITYSHADERS_API void ClearUAV(FRHICommandList& RHICmdList, const FSceneRenderTargetItem& RenderTargetItem, const float(&ClearValues)[4]);
+extern UTILITYSHADERS_API void ClearUAV(FRHICommandList& RHICmdList, const FSceneRenderTargetItem& RenderTargetItem, const uint32(&ClearValues)[4]);
+extern UTILITYSHADERS_API void ClearUAV(FRHICommandList& RHICmdList, const FSceneRenderTargetItem& RenderTargetItem,  const FLinearColor& ClearColor);
 
-extern UTILITYSHADERS_API void ClearTexture2DUAV(FRHICommandList& RHICmdList, ERHIFeatureLevel::Type FeatureLevel, FRHIUnorderedAccessView* UAV, int32 Width, int32 Height, const FLinearColor& ClearColor);
+extern UTILITYSHADERS_API void ClearTexture2DUAV(FRHICommandList& RHICmdList, FRHIUnorderedAccessView* UAV, int32 Width, int32 Height, const FLinearColor& ClearColor);
 
 
-extern UTILITYSHADERS_API void DrawClearQuadMRT( FRHICommandList& RHICmdList, ERHIFeatureLevel::Type FeatureLevel, bool bClearColor, int32 NumClearColors, const FLinearColor* ClearColorArray, bool bClearDepth, float Depth, bool bClearStencil, uint32 Stencil );
-extern UTILITYSHADERS_API void DrawClearQuadMRT( FRHICommandList& RHICmdList, ERHIFeatureLevel::Type FeatureLevel, bool bClearColor, int32 NumClearColors, const FLinearColor* ClearColorArray, bool bClearDepth, float Depth, bool bClearStencil, uint32 Stencil, FIntPoint ViewSize, FIntRect ExcludeRect );
+extern UTILITYSHADERS_API void DrawClearQuadMRT( FRHICommandList& RHICmdList, bool bClearColor, int32 NumClearColors, const FLinearColor* ClearColorArray, bool bClearDepth, float Depth, bool bClearStencil, uint32 Stencil );
+extern UTILITYSHADERS_API void DrawClearQuadMRT( FRHICommandList& RHICmdList, bool bClearColor, int32 NumClearColors, const FLinearColor* ClearColorArray, bool bClearDepth, float Depth, bool bClearStencil, uint32 Stencil, FIntPoint ViewSize, FIntRect ExcludeRect );
 
-inline void DrawClearQuad(FRHICommandList& RHICmdList, ERHIFeatureLevel::Type FeatureLevel, bool bClearColor, const FLinearColor& Color, bool bClearDepth, float Depth, bool bClearStencil, uint32 Stencil)
+inline void DrawClearQuad(FRHICommandList& RHICmdList, bool bClearColor, const FLinearColor& Color, bool bClearDepth, float Depth, bool bClearStencil, uint32 Stencil)
 {
-	DrawClearQuadMRT(RHICmdList, FeatureLevel, bClearColor, 1, &Color, bClearDepth, Depth, bClearStencil, Stencil);
+	DrawClearQuadMRT(RHICmdList, bClearColor, 1, &Color, bClearDepth, Depth, bClearStencil, Stencil);
 }
 
-inline void DrawClearQuad(FRHICommandList& RHICmdList, ERHIFeatureLevel::Type FeatureLevel, bool bClearColor, const FLinearColor& Color, bool bClearDepth, float Depth, bool bClearStencil, uint32 Stencil, FIntPoint ViewSize, FIntRect ExcludeRect)
+inline void DrawClearQuad(FRHICommandList& RHICmdList, bool bClearColor, const FLinearColor& Color, bool bClearDepth, float Depth, bool bClearStencil, uint32 Stencil, FIntPoint ViewSize, FIntRect ExcludeRect)
 {
-	DrawClearQuadMRT(RHICmdList, FeatureLevel, bClearColor, 1, &Color, bClearDepth, Depth, bClearStencil, Stencil, ViewSize, ExcludeRect);
+	DrawClearQuadMRT(RHICmdList, bClearColor, 1, &Color, bClearDepth, Depth, bClearStencil, Stencil, ViewSize, ExcludeRect);
 }
 
-inline void DrawClearQuad(FRHICommandList& RHICmdList, ERHIFeatureLevel::Type FeatureLevel, const FLinearColor& Color)
+inline void DrawClearQuad(FRHICommandList& RHICmdList, const FLinearColor& Color)
 {
-	DrawClearQuadMRT(RHICmdList, FeatureLevel, true, 1, &Color, false, 0, false, 0);
+	DrawClearQuadMRT(RHICmdList, true, 1, &Color, false, 0, false, 0);
 }
\ No newline at end of file
diff --git a/Engine/Source/Runtime/VulkanRHI/Private/VulkanMemory.cpp b/Engine/Source/Runtime/VulkanRHI/Private/VulkanMemory.cpp
--- a/Engine/Source/Runtime/VulkanRHI/Private/VulkanMemory.cpp
+++ b/Engine/Source/Runtime/VulkanRHI/Private/VulkanMemory.cpp
@@ -1450,7 +1450,7 @@ namespace VulkanRHI
 				VKSWITCH(Sampler);
 				VKSWITCH(Semaphore);
 				VKSWITCH(ShaderModule);
-#undef VKSWTICH
+#undef VKSWITCH
 				default:
 					check(0);
 					break;
diff --git a/Engine/Source/Runtime/VulkanRHI/Private/VulkanViewport.cpp b/Engine/Source/Runtime/VulkanRHI/Private/VulkanViewport.cpp
--- a/Engine/Source/Runtime/VulkanRHI/Private/VulkanViewport.cpp
+++ b/Engine/Source/Runtime/VulkanRHI/Private/VulkanViewport.cpp
@@ -122,7 +122,7 @@ FVulkanTexture2D* FVulkanViewport::GetBackBuffer(FRHICommandList& RHICmdList)
 		RenderingBackBuffer = new FVulkanBackBuffer(*Device, PixelFormat, SizeX, SizeY, VK_NULL_HANDLE, TexCreate_Presentable | TexCreate_RenderTargetable);
 		check(RHICmdList.IsImmediate());
 
-		if (RHICmdList.Bypass() || !GRHIThread)
+		if (RHICmdList.Bypass() || !IsRunningRHIInSeparateThread())
 		{
 			FRHICommandAcquireBackBuffer Cmd(this, RenderingBackBuffer);
 			Cmd.Execute(RHICmdList);
@@ -707,7 +707,7 @@ void FVulkanDynamicRHI::RHIAdvanceFrameForGetViewportBackBuffer()
 
 	{
 		FRHICommandList& RHICmdList = FRHICommandListExecutor::GetImmediateCommandList();
-		if (RHICmdList.Bypass() || !GRHIThread)
+		if (RHICmdList.Bypass() || !IsRunningRHIInSeparateThread())
 		{
 			FRHICommandProcessDeferredDeletionQueue Cmd(Device);
 			Cmd.Execute(RHICmdList);
diff --git a/Engine/Source/Runtime/Windows/D3D11RHI/Private/D3D11RHI.cpp b/Engine/Source/Runtime/Windows/D3D11RHI/Private/D3D11RHI.cpp
--- a/Engine/Source/Runtime/Windows/D3D11RHI/Private/D3D11RHI.cpp
+++ b/Engine/Source/Runtime/Windows/D3D11RHI/Private/D3D11RHI.cpp
@@ -132,7 +132,7 @@ void GetMipAndSliceInfoFromSRV(ID3D11ShaderResourceView* SRV, int32& MipLevel, i
 void FD3D11DynamicRHI::CheckIfSRVIsResolved(ID3D11ShaderResourceView* SRV)
 {
 #if CHECK_SRV_TRANSITIONS
-	if (GRHIThread || !SRV)
+	if (IsRunningRHIInSeparateThread() || !SRV)
 	{
 		return;
 	}
diff --git a/Engine/Source/Runtime/Windows/D3D11RHI/Private/D3D11RenderTarget.cpp b/Engine/Source/Runtime/Windows/D3D11RHI/Private/D3D11RenderTarget.cpp
--- a/Engine/Source/Runtime/Windows/D3D11RHI/Private/D3D11RenderTarget.cpp
+++ b/Engine/Source/Runtime/Windows/D3D11RHI/Private/D3D11RenderTarget.cpp
@@ -480,6 +480,7 @@ static uint32 ComputeBytesPerPixel(DXGI_FORMAT Format)
 		case DXGI_FORMAT_R10G10B10A2_UNORM:
 		case DXGI_FORMAT_R11G11B10_FLOAT:
 		case DXGI_FORMAT_R16G16_UNORM:
+		case DXGI_FORMAT_R32_UINT:
 			BytesPerPixel = 4;
 			break;
 		case DXGI_FORMAT_R16G16B16A16_FLOAT:
@@ -1272,44 +1273,62 @@ static void ConvertRAWSurfaceDataToFLinearColor(EPixelFormat Format, uint32 Widt
 	}
 	else if (Format == PF_FloatRGBA)
 	{
-		FPlane	MinValue(0.0f, 0.0f, 0.0f, 0.0f),
-			MaxValue(1.0f, 1.0f, 1.0f, 1.0f);
-
-		check(sizeof(FD3DFloat16) == sizeof(uint16));
-
-		for (uint32 Y = 0; Y < Height; Y++)
+		if (InFlags.GetCompressionMode() == RCM_MinMax)
 		{
-			FD3DFloat16* SrcPtr = (FD3DFloat16*)(In + Y * SrcPitch);
-
-			for (uint32 X = 0; X < Width; X++)
+			for (uint32 Y = 0; Y < Height; Y++)
 			{
-				MinValue.X = FMath::Min<float>(SrcPtr[0], MinValue.X);
-				MinValue.Y = FMath::Min<float>(SrcPtr[1], MinValue.Y);
-				MinValue.Z = FMath::Min<float>(SrcPtr[2], MinValue.Z);
-				MinValue.W = FMath::Min<float>(SrcPtr[3], MinValue.W);
-				MaxValue.X = FMath::Max<float>(SrcPtr[0], MaxValue.X);
-				MaxValue.Y = FMath::Max<float>(SrcPtr[1], MaxValue.Y);
-				MaxValue.Z = FMath::Max<float>(SrcPtr[2], MaxValue.Z);
-				MaxValue.W = FMath::Max<float>(SrcPtr[3], MaxValue.W);
-				SrcPtr += 4;
+				FD3DFloat16* SrcPtr = (FD3DFloat16*)(In + Y * SrcPitch);
+				FLinearColor* DestPtr = Out + Y * Width;
+
+				for (uint32 X = 0; X < Width; X++)
+				{
+					*DestPtr = FLinearColor((float)SrcPtr[0], (float)SrcPtr[1], (float)SrcPtr[2], (float)SrcPtr[3]);
+					++DestPtr;
+					SrcPtr += 4;
+				}
 			}
 		}
-
-		for (uint32 Y = 0; Y < Height; Y++)
+		else
 		{
-			FD3DFloat16* SrcPtr = (FD3DFloat16*)(In + Y * SrcPitch);
-			FLinearColor* DestPtr = Out + Y * Width;
+			FPlane	MinValue(0.0f, 0.0f, 0.0f, 0.0f);
+			FPlane	MaxValue(1.0f, 1.0f, 1.0f, 1.0f);
 
-			for (uint32 X = 0; X < Width; X++)
+			check(sizeof(FD3DFloat16) == sizeof(uint16));
+
+			for (uint32 Y = 0; Y < Height; Y++)
 			{
-				*DestPtr = FLinearColor(
-					(SrcPtr[0] - MinValue.X) / (MaxValue.X - MinValue.X),
-					(SrcPtr[1] - MinValue.Y) / (MaxValue.Y - MinValue.Y),
-					(SrcPtr[2] - MinValue.Z) / (MaxValue.Z - MinValue.Z),
-					(SrcPtr[3] - MinValue.W) / (MaxValue.W - MinValue.W)
-					);
-				++DestPtr;
-				SrcPtr += 4;
+				FD3DFloat16* SrcPtr = (FD3DFloat16*)(In + Y * SrcPitch);
+
+				for (uint32 X = 0; X < Width; X++)
+				{
+					MinValue.X = FMath::Min<float>(SrcPtr[0], MinValue.X);
+					MinValue.Y = FMath::Min<float>(SrcPtr[1], MinValue.Y);
+					MinValue.Z = FMath::Min<float>(SrcPtr[2], MinValue.Z);
+					MinValue.W = FMath::Min<float>(SrcPtr[3], MinValue.W);
+					MaxValue.X = FMath::Max<float>(SrcPtr[0], MaxValue.X);
+					MaxValue.Y = FMath::Max<float>(SrcPtr[1], MaxValue.Y);
+					MaxValue.Z = FMath::Max<float>(SrcPtr[2], MaxValue.Z);
+					MaxValue.W = FMath::Max<float>(SrcPtr[3], MaxValue.W);
+					SrcPtr += 4;
+				}
+			}
+
+			for (uint32 Y = 0; Y < Height; Y++)
+			{
+				FD3DFloat16* SrcPtr = (FD3DFloat16*)(In + Y * SrcPitch);
+				FLinearColor* DestPtr = Out + Y * Width;
+
+				for (uint32 X = 0; X < Width; X++)
+				{
+					*DestPtr = FLinearColor(
+						(SrcPtr[0] - MinValue.X) / (MaxValue.X - MinValue.X),
+						(SrcPtr[1] - MinValue.Y) / (MaxValue.Y - MinValue.Y),
+						(SrcPtr[2] - MinValue.Z) / (MaxValue.Z - MinValue.Z),
+						(SrcPtr[3] - MinValue.W) / (MaxValue.W - MinValue.W)
+						);
+					++DestPtr;
+					SrcPtr += 4;
+				}
 			}
 		}
 	}
diff --git a/Engine/Source/Runtime/Windows/D3D11RHI/Private/D3D11UniformBuffer.cpp b/Engine/Source/Runtime/Windows/D3D11RHI/Private/D3D11UniformBuffer.cpp
--- a/Engine/Source/Runtime/Windows/D3D11RHI/Private/D3D11UniformBuffer.cpp
+++ b/Engine/Source/Runtime/Windows/D3D11RHI/Private/D3D11UniformBuffer.cpp
@@ -95,7 +95,7 @@ void UniformBufferBeginFrame()
 
 static bool IsPoolingEnabled()
 {
-	if (GRHIThread && IsInRenderingThread() && GRHICommandList.IsRHIThreadActive())
+	if (IsRunningRHIInSeparateThread() && IsInRenderingThread() && GRHICommandList.IsRHIThreadActive())
 	{
 		return false; // we can't currently use pooling if the RHI thread is active. 
 	}
diff --git a/Engine/Source/Runtime/Windows/D3D11RHI/Private/Windows/WindowsD3D11Device.cpp b/Engine/Source/Runtime/Windows/D3D11RHI/Private/Windows/WindowsD3D11Device.cpp
--- a/Engine/Source/Runtime/Windows/D3D11RHI/Private/Windows/WindowsD3D11Device.cpp
+++ b/Engine/Source/Runtime/Windows/D3D11RHI/Private/Windows/WindowsD3D11Device.cpp
@@ -30,6 +30,26 @@ static FAutoConsoleVariableRef CVarDX11NVAfterMathBufferSize(
 extern bool D3D11RHI_ShouldCreateWithD3DDebug();
 extern bool D3D11RHI_ShouldAllowAsyncResourceCreation();
 
+int D3D11RHI_PreferAdaperVendor()
+{
+	if (FParse::Param(FCommandLine::Get(), TEXT("preferAMD")))
+	{
+		return 0x1002;
+	}
+
+	if (FParse::Param(FCommandLine::Get(), TEXT("preferIntel")))
+	{
+		return 0x8086;
+	}
+
+	if (FParse::Param(FCommandLine::Get(), TEXT("preferNvidia")))
+	{
+		return 0x10DE;
+	}
+
+	return -1;
+}
+
 // Filled in during InitD3DDevice if IsRHIDeviceAMD
 struct AmdAgsInfo
 {
@@ -644,6 +664,7 @@ void FD3D11DynamicRHIModule::FindAdapter()
 
 	UE_LOG(LogD3D11RHI, Log, TEXT("D3D11 adapters:"));
 
+	int PreferredVendor = D3D11RHI_PreferAdaperVendor();
 	// Enumerate the DXGIFactory's adapters.
 	for(uint32 AdapterIndex = 0; DXGIFactory1->EnumAdapters(AdapterIndex,TempAdapter.GetInitReference()) != DXGI_ERROR_NOT_FOUND; ++AdapterIndex)
 	{
@@ -711,11 +732,19 @@ void FD3D11DynamicRHIModule::FindAdapter()
 					{
 						FirstWithoutIntegratedAdapter = CurrentAdapter;
 					}
+					else if (PreferredVendor == AdapterDesc.VendorId && FirstWithoutIntegratedAdapter.IsValid())
+					{
+						FirstWithoutIntegratedAdapter = CurrentAdapter;
+					}
 
 					if (!FirstAdapter.IsValid())
 					{
 						FirstAdapter = CurrentAdapter;
 					}
+					else if (PreferredVendor == AdapterDesc.VendorId && FirstAdapter.IsValid())
+					{
+						FirstAdapter = CurrentAdapter;
+					}
 				}
 			}
 		}
@@ -1097,7 +1126,11 @@ void FD3D11DynamicRHI::InitD3DDevice()
 				}
 				else
 				{
-					UE_LOG(LogD3D11RHI, Log, TEXT("[Aftermath] Aftermath enabled but failed to initialize"));
+					unsigned Index = (unsigned)Result & (~((unsigned)GFSDK_Aftermath_Result_Fail));
+					const TCHAR* Reason[13] = { TEXT("Fail"), TEXT("VersionMismatch"), TEXT("NotInitialized"), TEXT("InvalidAdapter"), TEXT("InvalidParameter"), TEXT("Unknown"), TEXT("ApiError"), TEXT("NvApiIncompatible"), TEXT("GettingContextDataWithNewCommandList"), TEXT("AlreadyInitialized"), TEXT("D3DDebugLayerNotCompatible"),TEXT("NotEnabledInDriver"), TEXT("DriverVersionNotSupported") };
+					Index = Index > 12 ? 0 : Index;
+
+					UE_LOG(LogD3D11RHI, Log, TEXT("[Aftermath] Aftermath enabled but failed to initialize due to reason: %s"), Reason[Index]);
 					GDX11NVAfterMathEnabled = 0;
 				}
 			}
