fix for attempting to render a clear quad when using a null rhi

diff --git a/Engine/Source/Runtime/UtilityShaders/Private/ClearQuad.cpp b/Engine/Source/Runtime/UtilityShaders/Private/ClearQuad.cpp
--- a/Engine/Source/Runtime/UtilityShaders/Private/ClearQuad.cpp
+++ b/Engine/Source/Runtime/UtilityShaders/Private/ClearQuad.cpp
@@ -10,6 +10,11 @@
 
 static void ClearQuadSetup( FRHICommandList& RHICmdList, bool bClearColor, int32 NumClearColors, const FLinearColor* ClearColorArray, bool bClearDepth, float Depth, bool bClearStencil, uint32 Stencil )
 {
+	if (UNLIKELY(!FApp::CanEverRender()))
+	{
+		return;
+	}
+
 	// Set new states
 	FBlendStateRHIParamRef BlendStateRHI;
 		
