UE-46661 Editor crashes when launching onto PS4 with Morpheus enabled due to shader compilation failure
-Fixing shader includes with the new power to put them in platform specific sub directories.

diff --git a/Engine/Shaders/Private/MorpheusInclude.usf b/Engine/Shaders/Private/MorpheusInclude.usf
deleted file mode 100644
diff --git a/Engine/Source/Runtime/Renderer/Private/PostProcess/PostProcessMorpheus.cpp b/Engine/Source/Runtime/Renderer/Private/PostProcess/PostProcessMorpheus.cpp
--- a/Engine/Source/Runtime/Renderer/Private/PostProcess/PostProcessMorpheus.cpp
+++ b/Engine/Source/Runtime/Renderer/Private/PostProcess/PostProcessMorpheus.cpp
@@ -185,8 +185,8 @@ public:
 	}
 };
 
-IMPLEMENT_SHADER_TYPE(, FPostProcessMorpheusVS, TEXT("/Engine/Private/MorpheusInclude.usf"), TEXT("MainVS"), SF_Vertex);
-IMPLEMENT_SHADER_TYPE(, FPostProcessMorpheusPS,TEXT("/Engine/Private/MorpheusInclude.usf"),TEXT("MainPS"),SF_Pixel);
+IMPLEMENT_SHADER_TYPE(, FPostProcessMorpheusVS, TEXT("/Engine/Private/PS4/PostProcessHMDMorpheus.usf"), TEXT("MainVS"), SF_Vertex);
+IMPLEMENT_SHADER_TYPE(, FPostProcessMorpheusPS,TEXT("/Engine/Private/PS4/PostProcessHMDMorpheus.usf"),TEXT("MainPS"),SF_Pixel);
 
 void FRCPassPostProcessMorpheus::Process(FRenderingCompositePassContext& Context)
 {
diff --git a/Engine/Source/Runtime/UtilityShaders/Private/RGBAToYUV420Shader.cpp b/Engine/Source/Runtime/UtilityShaders/Private/RGBAToYUV420Shader.cpp
--- a/Engine/Source/Runtime/UtilityShaders/Private/RGBAToYUV420Shader.cpp
+++ b/Engine/Source/Runtime/UtilityShaders/Private/RGBAToYUV420Shader.cpp
@@ -17,7 +17,7 @@ BEGIN_UNIFORM_BUFFER_STRUCT( RGBAToYUV420UB, )
 END_UNIFORM_BUFFER_STRUCT( RGBAToYUV420UB )
 
 IMPLEMENT_UNIFORM_BUFFER_STRUCT(RGBAToYUV420UB, TEXT("RGBAToYUV420UB"));
-IMPLEMENT_SHADER_TYPE(, FRGBAToYUV420CS, TEXT("/Engine/Private/MorpheusInclude.usf"), TEXT("RGBAToYUV420Main"), SF_Compute);
+IMPLEMENT_SHADER_TYPE(, FRGBAToYUV420CS, TEXT("/Engine/Private/PS4/RGBAToYUV420.usf"), TEXT("RGBAToYUV420Main"), SF_Compute);
 
 
 void FRGBAToYUV420CS::SetParameters(FRHICommandList& RHICmdList, TRefCountPtr<FRHITexture2D> SrcTex, FUnorderedAccessViewRHIParamRef OutUAV, float TargetHeight, float ScaleFactorX, float ScaleFactorY, float TextureYOffset)
