Fixes yet another shader include parsing failure caused by 3517032 on RGBAToYUV420.usf. The includer scanner was failing because does not have the logic for file ignoring based on whether it is in PS4 or XboxOne directories.

diff --git a/Engine/Shaders/Private/RGBAToYUV420.usf b/Engine/Shaders/Private/RGBAToYUV420.usf
new file mode 100644
--- /dev/null
+++ b/Engine/Shaders/Private/RGBAToYUV420.usf
@@ -0,0 +1,9 @@
+// Copyright 1998-2017 Epic Games, Inc. All Rights Reserved.
+
+// PS4 directory may not exists. As a result a shader compiler may not be able to
+// find file for include scanning, but the include scanner will silently ignore
+// PS4 directory by design.
+//
+// TODO(VR team): move shader file and C++ into a plugin.
+
+#include "PS4/RGBAToYUV420.usf"
diff --git a/Engine/Source/Runtime/UtilityShaders/Private/RGBAToYUV420Shader.cpp b/Engine/Source/Runtime/UtilityShaders/Private/RGBAToYUV420Shader.cpp
--- a/Engine/Source/Runtime/UtilityShaders/Private/RGBAToYUV420Shader.cpp
+++ b/Engine/Source/Runtime/UtilityShaders/Private/RGBAToYUV420Shader.cpp
@@ -17,7 +17,7 @@ BEGIN_UNIFORM_BUFFER_STRUCT( RGBAToYUV420UB, )
 END_UNIFORM_BUFFER_STRUCT( RGBAToYUV420UB )
 
 IMPLEMENT_UNIFORM_BUFFER_STRUCT(RGBAToYUV420UB, TEXT("RGBAToYUV420UB"));
-IMPLEMENT_SHADER_TYPE(, FRGBAToYUV420CS, TEXT("/Engine/Private/PS4/RGBAToYUV420.usf"), TEXT("RGBAToYUV420Main"), SF_Compute);
+IMPLEMENT_SHADER_TYPE(, FRGBAToYUV420CS, TEXT("/Engine/Private/RGBAToYUV420.usf"), TEXT("RGBAToYUV420Main"), SF_Compute);
 
 
 void FRGBAToYUV420CS::SetParameters(FRHICommandList& RHICmdList, TRefCountPtr<FRHITexture2D> SrcTex, FUnorderedAccessViewRHIParamRef OutUAV, float TargetHeight, float ScaleFactorX, float ScaleFactorY, float TextureYOffset)
