UE-46798 Crash launching onto PS4 with Morpheus enabled
-duplicating from Release-4.17 cl 3525417
-Removed unnecessary includes.  Switched from Uniform Buffer to FShaderParameter param passing.  UB is unnecessary if it is only used by one shader.
-Switched the RGBAToYUV420.usf file to a .ush file.  This seemed to help avoid it being built twice.

diff --git a/Engine/Shaders/Private/RGBAToYUV420.usf b/Engine/Shaders/Private/RGBAToYUV420.usf
--- a/Engine/Shaders/Private/RGBAToYUV420.usf
+++ b/Engine/Shaders/Private/RGBAToYUV420.usf
@@ -6,4 +6,4 @@
 //
 // TODO(VR team): move shader file and C++ into a plugin.
 
-#include "PS4/RGBAToYUV420.usf"
+#include "PS4/RGBAToYUV420.ush"
diff --git a/Engine/Source/Runtime/ShaderCore/Private/ShaderCore.cpp b/Engine/Source/Runtime/ShaderCore/Private/ShaderCore.cpp
--- a/Engine/Source/Runtime/ShaderCore/Private/ShaderCore.cpp
+++ b/Engine/Source/Runtime/ShaderCore/Private/ShaderCore.cpp
@@ -507,7 +507,7 @@ void GetShaderIncludes(const TCHAR* EntryPointVirtualFilePath, const TCHAR* Virt
 					const bool bIsOptionalInclude = (ExtractedIncludeFilename == TEXT("/Engine/Public/PS4/PS4Common.ush") 
 						|| ExtractedIncludeFilename == TEXT("/Engine/Private/PS4/PostProcessHMDMorpheus.usf")
 						|| ExtractedIncludeFilename == TEXT("/Engine/Private/PS4/RTWriteMaskProcessing.usf")
-						|| ExtractedIncludeFilename == TEXT("/Engine/Private/PS4/RGBAToYUV420.usf")
+						|| ExtractedIncludeFilename == TEXT("/Engine/Private/PS4/RGBAToYUV420.ush")
 						|| ExtractedIncludeFilename == TEXT("/Engine/Public/XboxOne/XboxOneCommon.ush")
 						);
 					// ignore the header if it's optional and doesn't exist
diff --git a/Engine/Source/Runtime/UtilityShaders/Private/RGBAToYUV420Shader.cpp b/Engine/Source/Runtime/UtilityShaders/Private/RGBAToYUV420Shader.cpp
--- a/Engine/Source/Runtime/UtilityShaders/Private/RGBAToYUV420Shader.cpp
+++ b/Engine/Source/Runtime/UtilityShaders/Private/RGBAToYUV420Shader.cpp
@@ -8,34 +8,18 @@
 /* FRGBAToYUV420CS shader
  *****************************************************************************/
 
-BEGIN_UNIFORM_BUFFER_STRUCT( RGBAToYUV420UB, )
-	DECLARE_UNIFORM_BUFFER_STRUCT_MEMBER( float, TargetHeight )
-	DECLARE_UNIFORM_BUFFER_STRUCT_MEMBER( float, ScaleFactorX )
-	DECLARE_UNIFORM_BUFFER_STRUCT_MEMBER( float, ScaleFactorY )
-	DECLARE_UNIFORM_BUFFER_STRUCT_MEMBER( float, TextureYOffset )
-	DECLARE_UNIFORM_BUFFER_STRUCT_MEMBER_TEXTURE( Texture2D, SrcTex )
-END_UNIFORM_BUFFER_STRUCT( RGBAToYUV420UB )
-
-IMPLEMENT_UNIFORM_BUFFER_STRUCT(RGBAToYUV420UB, TEXT("RGBAToYUV420UB"));
 IMPLEMENT_SHADER_TYPE(, FRGBAToYUV420CS, TEXT("/Engine/Private/RGBAToYUV420.usf"), TEXT("RGBAToYUV420Main"), SF_Compute);
 
 
-void FRGBAToYUV420CS::SetParameters(FRHICommandList& RHICmdList, TRefCountPtr<FRHITexture2D> SrcTex, FUnorderedAccessViewRHIParamRef OutUAV, float TargetHeight, float ScaleFactorX, float ScaleFactorY, float TextureYOffset)
+void FRGBAToYUV420CS::SetParameters(FRHICommandList& RHICmdList, TRefCountPtr<FRHITexture2D> InSrcTexture, FUnorderedAccessViewRHIParamRef InOutUAV, float InTargetHeight, float InScaleFactorX, float InScaleFactorY, float InTextureYOffset)
 {
-	RGBAToYUV420UB UB;
-	{
-		UB.SrcTex = SrcTex;
-		UB.TargetHeight = TargetHeight;
-		UB.ScaleFactorX = ScaleFactorX;
-		UB.ScaleFactorY = ScaleFactorY;
-		UB.TextureYOffset = TextureYOffset;
-	}
 	FComputeShaderRHIParamRef ComputeShaderRHI = GetComputeShader();
-
-	TUniformBufferRef<RGBAToYUV420UB> Data = TUniformBufferRef<RGBAToYUV420UB>::CreateUniformBufferImmediate(UB, UniformBuffer_SingleFrame);
-	SetUniformBufferParameter(RHICmdList, ComputeShaderRHI, GetUniformBufferParameter<RGBAToYUV420UB>(), Data);
-
-	RHICmdList.SetUAVParameter(ComputeShaderRHI, OutTextureRW.GetBaseIndex(), OutUAV);
+	SetShaderValue(RHICmdList, ComputeShaderRHI, TargetHeight, InTargetHeight);
+	SetShaderValue(RHICmdList, ComputeShaderRHI, ScaleFactorX, InScaleFactorX);
+	SetShaderValue(RHICmdList, ComputeShaderRHI, ScaleFactorY, InScaleFactorY);
+	SetShaderValue(RHICmdList, ComputeShaderRHI, TextureYOffset, InTextureYOffset);
+	SetTextureParameter(RHICmdList, ComputeShaderRHI, SrcTexture, InSrcTexture);
+	RHICmdList.SetUAVParameter(ComputeShaderRHI, OutTextureRW.GetBaseIndex(), InOutUAV);
 }
 
 /**
diff --git a/Engine/Source/Runtime/UtilityShaders/Public/RGBAToYUV420Shader.h b/Engine/Source/Runtime/UtilityShaders/Public/RGBAToYUV420Shader.h
--- a/Engine/Source/Runtime/UtilityShaders/Public/RGBAToYUV420Shader.h
+++ b/Engine/Source/Runtime/UtilityShaders/Public/RGBAToYUV420Shader.h
@@ -43,12 +43,18 @@ public:
 	FRGBAToYUV420CS(const ShaderMetaType::CompiledShaderInitializerType& Initializer)
 		: FGlobalShader(Initializer)
 	{
+		TargetHeight.Bind(Initializer.ParameterMap, TEXT("TargetHeight"), SPF_Mandatory);
+		ScaleFactorX.Bind(Initializer.ParameterMap, TEXT("ScaleFactorX"), SPF_Mandatory);
+		ScaleFactorY.Bind(Initializer.ParameterMap, TEXT("ScaleFactorY"), SPF_Mandatory);
+		TextureYOffset.Bind(Initializer.ParameterMap, TEXT("TextureYOffset"), SPF_Mandatory);
+		SrcTexture.Bind(Initializer.ParameterMap, TEXT("SrcTexture"), SPF_Mandatory);
 		OutTextureRW.Bind(Initializer.ParameterMap, TEXT("OutTexture"), SPF_Mandatory);
 	}
 
 	virtual bool Serialize(FArchive& Ar) override
 	{
 		bool bShaderHasOutdatedParameters = FGlobalShader::Serialize(Ar);
+		Ar << TargetHeight << ScaleFactorX << ScaleFactorY << TextureYOffset << SrcTexture << OutTextureRW;
 		return bShaderHasOutdatedParameters;
 	}
 
@@ -56,6 +62,12 @@ public:
 	void UnbindBuffers(FRHICommandList& RHICmdList);
 
 protected:
+	FShaderParameter TargetHeight;
+	FShaderParameter ScaleFactorX;
+	FShaderParameter ScaleFactorY;
+	FShaderParameter TextureYOffset;
+	FShaderResourceParameter SrcTexture;
+
 	FShaderResourceParameter OutTextureRW;
 };
 
