Fixed: DepthFade does not work in mobile emulation while MobileMSAA enabled

diff --git a/Engine/Shaders/Private/ResolvePixelShader.usf b/Engine/Shaders/Private/ResolvePixelShader.usf
--- a/Engine/Shaders/Private/ResolvePixelShader.usf
+++ b/Engine/Shaders/Private/ResolvePixelShader.usf
@@ -6,7 +6,7 @@
 
 #include "Common.ush"
 
-#if FEATURE_LEVEL >= FEATURE_LEVEL_SM5
+#if FEATURE_LEVEL >= FEATURE_LEVEL_SM5 || SIMULATED_PLATFORM  
 
 Texture2DMS<float4> UnresolvedSurface;
 
diff --git a/Engine/Source/Developer/Windows/ShaderFormatD3D/Private/D3D11ShaderCompiler.cpp b/Engine/Source/Developer/Windows/ShaderFormatD3D/Private/D3D11ShaderCompiler.cpp
--- a/Engine/Source/Developer/Windows/ShaderFormatD3D/Private/D3D11ShaderCompiler.cpp
+++ b/Engine/Source/Developer/Windows/ShaderFormatD3D/Private/D3D11ShaderCompiler.cpp
@@ -133,7 +133,7 @@ static const TCHAR* GetShaderProfileName(FShaderTarget Target)
 			return TEXT("cs_5_0");
 		}
 	}
-	else if(Target.Platform == SP_PCD3D_SM4 || Target.Platform == SP_PCD3D_ES2 || Target.Platform == SP_PCD3D_ES3_1)
+	else if(Target.Platform == SP_PCD3D_SM4)
 	{
 		checkSlow(Target.Frequency == SF_Vertex ||
 			Target.Frequency == SF_Pixel ||
@@ -150,6 +150,23 @@ static const TCHAR* GetShaderProfileName(FShaderTarget Target)
 			return TEXT("gs_4_0");
 		}
 	}
+	else if (Target.Platform == SP_PCD3D_ES2 || Target.Platform == SP_PCD3D_ES3_1)
+	{
+		checkSlow(Target.Frequency == SF_Vertex ||
+			Target.Frequency == SF_Pixel ||
+			Target.Frequency == SF_Geometry);
+
+		//set defines and profiles for the appropriate shader paths
+		switch(Target.Frequency)
+		{
+		case SF_Pixel:
+			return TEXT("ps_5_0");
+		case SF_Vertex:
+			return TEXT("vs_5_0");
+		case SF_Geometry:
+			return TEXT("gs_5_0");
+		}
+	}
 
 	return NULL;
 }
diff --git a/Engine/Source/Developer/Windows/ShaderFormatD3D/Private/ShaderFormatD3D.cpp b/Engine/Source/Developer/Windows/ShaderFormatD3D/Private/ShaderFormatD3D.cpp
--- a/Engine/Source/Developer/Windows/ShaderFormatD3D/Private/ShaderFormatD3D.cpp
+++ b/Engine/Source/Developer/Windows/ShaderFormatD3D/Private/ShaderFormatD3D.cpp
@@ -18,8 +18,8 @@ class FShaderFormatD3D : public IShaderFormat
 		/** Version for shader format, this becomes part of the DDC key. */
 		UE_SHADER_PCD3D_SM5_VER = 7,
 		UE_SHADER_PCD3D_SM4_VER = 7,
-		UE_SHADER_PCD3D_ES2_VER = 7,
-		UE_SHADER_PCD3D_ES3_1_VER = 7,
+		UE_SHADER_PCD3D_ES2_VER = 8,
+		UE_SHADER_PCD3D_ES3_1_VER = 8,
 	};
 
 	void CheckFormat(FName Format) const
diff --git a/Engine/Source/Runtime/RHI/Public/RHIDefinitions.h b/Engine/Source/Runtime/RHI/Public/RHIDefinitions.h
--- a/Engine/Source/Runtime/RHI/Public/RHIDefinitions.h
+++ b/Engine/Source/Runtime/RHI/Public/RHIDefinitions.h
@@ -939,6 +939,9 @@ inline bool IsSimulatedPlatform(EShaderPlatform Platform)
 		case SP_PCD3D_ES2:
 		case SP_PCD3D_ES3_1:
 		case SP_OPENGL_PCES3_1:
+		case SP_METAL_MACES3_1:
+		case SP_METAL_MACES2:
+		case SP_VULKAN_PCES3_1:
 			return true;
 		break;
 
diff --git a/Engine/Source/Runtime/Renderer/Private/MobileShadingRenderer.cpp b/Engine/Source/Runtime/Renderer/Private/MobileShadingRenderer.cpp
--- a/Engine/Source/Runtime/Renderer/Private/MobileShadingRenderer.cpp
+++ b/Engine/Source/Runtime/Renderer/Private/MobileShadingRenderer.cpp
@@ -555,23 +555,28 @@ int32 FMobileSceneRenderer::ComputeNumOcclusionQueriesToBatch() const
 void FMobileSceneRenderer::ConditionalResolveSceneDepth(FRHICommandListImmediate& RHICmdList, const FViewInfo& View)
 {
 	FSceneRenderTargets& SceneContext = FSceneRenderTargets::Get(RHICmdList);
-	
 	SceneContext.ResolveSceneDepthToAuxiliaryTexture(RHICmdList);
-
-	auto ShaderPlatform = ViewFamily.GetShaderPlatform();
-
-	if ((IsMobileHDR() || IsHTML5Platform())
+	
+	EShaderPlatform ShaderPlatform = ViewFamily.GetShaderPlatform();
+	if (IsSimulatedPlatform(ShaderPlatform)) // mobile emulation on PC
+	{
+		// resolve MSAA depth for translucency
+		SceneContext.ResolveSceneDepthTexture(RHICmdList, FResolveRect(0, 0, FamilySize.X, FamilySize.Y));
+	}
+	else if ((IsMobileHDR() || IsHTML5Platform())
 		&& IsMobilePlatform(ShaderPlatform) 
 		&& !IsVulkanPlatform(ShaderPlatform)
 		&& !IsMetalPlatform(ShaderPlatform)
-		&& !IsPCPlatform(ShaderPlatform) // exclude mobile emulation on PC
 		&& !View.bIsPlanarReflection)	// exclude depth resolve from planar reflection captures, can't do it reliably more than once per frame
 	{
 		bool bSceneDepthInAlpha = (SceneContext.GetSceneColor()->GetDesc().Format == PF_FloatRGBA);
 		bool bOnChipDepthFetch = (GSupportsShaderDepthStencilFetch || (GSupportsShaderFramebufferFetch && bSceneDepthInAlpha));
-		
-		const bool bAlwaysResolveDepth = CVarMobileAlwaysResolveDepth.GetValueOnRenderThread() == 1;
 
+		const bool bAlwaysResolveDepth = CVarMobileAlwaysResolveDepth.GetValueOnRenderThread() == 1;
+		
+		static const auto CVarMobileMSAA = IConsoleManager::Get().FindTConsoleVariableDataInt(TEXT("r.MobileMSAA"));
+		const bool bMobileMSAA = (CVarMobileMSAA ? CVarMobileMSAA->GetValueOnAnyThread() > 1 : false);
+		
 		if (!bOnChipDepthFetch || bAlwaysResolveDepth)
 		{
 			// Only these features require depth texture
@@ -590,7 +595,7 @@ void FMobileSceneRenderer::ConditionalResolveSceneDepth(FRHICommandListImmediate
 						CopySceneAlpha(RHICmdList, View);
 					}
 				}
-				else
+				else if (!bMobileMSAA)
 				{
 					// Switch target to force hardware flush current depth to texture
 					FTextureRHIRef DummySceneColor = GSystemTextures.BlackDummy->GetRenderTargetItem().TargetableTexture;
diff --git a/Engine/Source/Runtime/UtilityShaders/Public/ResolveShader.h b/Engine/Source/Runtime/UtilityShaders/Public/ResolveShader.h
--- a/Engine/Source/Runtime/UtilityShaders/Public/ResolveShader.h
+++ b/Engine/Source/Runtime/UtilityShaders/Public/ResolveShader.h
@@ -17,8 +17,18 @@ public:
 	
 	typedef FDummyResolveParameter FParameter;
 	
-	static bool ShouldCompilePermutation(const FGlobalShaderPermutationParameters& Parameters) { return GetMaxSupportedFeatureLevel(Parameters.Platform) >= ERHIFeatureLevel::SM5; }
-	
+	static bool ShouldCompilePermutation(const FGlobalShaderPermutationParameters& Parameters) 
+	{ 
+		return GetMaxSupportedFeatureLevel(Parameters.Platform) >= ERHIFeatureLevel::SM5
+			|| IsSimulatedPlatform(Parameters.Platform); // support resolving MSAA depth in mobile emulation
+	}
+
+	static void ModifyCompilationEnvironment(const FGlobalShaderPermutationParameters& Parameters, FShaderCompilerEnvironment& OutEnvironment)
+	{
+		FGlobalShader::ModifyCompilationEnvironment(Parameters, OutEnvironment);
+		OutEnvironment.SetDefine(TEXT("SIMULATED_PLATFORM"), IsSimulatedPlatform(Parameters.Platform) ? 1 : 0);
+	}
+		
 	FResolveDepthPS(const ShaderMetaType::CompiledShaderInitializerType& Initializer):
 	FGlobalShader(Initializer)
 	{
@@ -47,12 +57,17 @@ public:
 
 	typedef FDummyResolveParameter FParameter;
 
-	static bool ShouldCompilePermutation(const FGlobalShaderPermutationParameters& Parameters) { return GetMaxSupportedFeatureLevel(Parameters.Platform) >= ERHIFeatureLevel::SM5; }
+	static bool ShouldCompilePermutation(const FGlobalShaderPermutationParameters& Parameters) 
+	{ 
+		return GetMaxSupportedFeatureLevel(Parameters.Platform) >= ERHIFeatureLevel::SM5 
+			|| IsSimulatedPlatform(Parameters.Platform); // support resolving MSAA depth in mobile emulation
+	}
 
 	static void ModifyCompilationEnvironment(const FGlobalShaderPermutationParameters& Parameters, FShaderCompilerEnvironment& OutEnvironment)
 	{
 		FGlobalShader::ModifyCompilationEnvironment(Parameters, OutEnvironment);
 		OutEnvironment.SetDefine(TEXT("DEPTH_RESOLVE_NUM_SAMPLES"), 2);
+		OutEnvironment.SetDefine(TEXT("SIMULATED_PLATFORM"), IsSimulatedPlatform(Parameters.Platform) ? 1 : 0);
 	}
 
 	FResolveDepth2XPS(const ShaderMetaType::CompiledShaderInitializerType& Initializer) :
@@ -84,12 +99,17 @@ public:
 
 	typedef FDummyResolveParameter FParameter;
 
-	static bool ShouldCompilePermutation(const FGlobalShaderPermutationParameters& Parameters) { return GetMaxSupportedFeatureLevel(Parameters.Platform) >= ERHIFeatureLevel::SM5; }
+	static bool ShouldCompilePermutation(const FGlobalShaderPermutationParameters& Parameters) 
+	{ 
+		return GetMaxSupportedFeatureLevel(Parameters.Platform) >= ERHIFeatureLevel::SM5
+			|| IsSimulatedPlatform(Parameters.Platform); // support resolving MSAA depth in mobile emulation
+	}
 
 	static void ModifyCompilationEnvironment(const FGlobalShaderPermutationParameters& Parameters, FShaderCompilerEnvironment& OutEnvironment)
 	{
 		FGlobalShader::ModifyCompilationEnvironment(Parameters, OutEnvironment);
 		OutEnvironment.SetDefine(TEXT("DEPTH_RESOLVE_NUM_SAMPLES"), 4);
+		OutEnvironment.SetDefine(TEXT("SIMULATED_PLATFORM"), IsSimulatedPlatform(Parameters.Platform) ? 1 : 0);
 	}
 
 	FResolveDepth4XPS(const ShaderMetaType::CompiledShaderInitializerType& Initializer) :
@@ -125,13 +145,15 @@ public:
 
 	static bool ShouldCompilePermutation(const FGlobalShaderPermutationParameters& Parameters)
 	{
-		return GetMaxSupportedFeatureLevel(Parameters.Platform) >= ERHIFeatureLevel::SM5;
+		return GetMaxSupportedFeatureLevel(Parameters.Platform) >= ERHIFeatureLevel::SM5
+			|| IsSimulatedPlatform(Parameters.Platform); // support resolving MSAA depth in mobile emulation
 	}
 
 	static void ModifyCompilationEnvironment(const FGlobalShaderPermutationParameters& Parameters, FShaderCompilerEnvironment& OutEnvironment)
 	{
 		FGlobalShader::ModifyCompilationEnvironment(Parameters, OutEnvironment);
 		OutEnvironment.SetDefine(TEXT("DEPTH_RESOLVE_NUM_SAMPLES"), 8);
+		OutEnvironment.SetDefine(TEXT("SIMULATED_PLATFORM"), IsSimulatedPlatform(Parameters.Platform) ? 1 : 0);
 	}
 
 	FResolveDepth8XPS(const ShaderMetaType::CompiledShaderInitializerType& Initializer) :
