ClearUAV calls RHI::ClearTinyUAV if buffer size to clear is lower than r.RHI.FastClearUAVMaxSize value

#rnx

diff --git a/Engine/Source/Runtime/UtilityShaders/Private/ClearQuad.cpp b/Engine/Source/Runtime/UtilityShaders/Private/ClearQuad.cpp
--- a/Engine/Source/Runtime/UtilityShaders/Private/ClearQuad.cpp
+++ b/Engine/Source/Runtime/UtilityShaders/Private/ClearQuad.cpp
@@ -11,6 +11,13 @@
 
 TGlobalResource<FClearVertexBuffer> GClearVertexBuffer;
 
+static TAutoConsoleVariable<int32> CVarFastClearUAVMaxSize(
+	TEXT("r.RHI.FastClearUAVMaxSize"),
+	0,
+	TEXT("Max size in bytes to fast clear tiny UAV. 0 by default or when feature is not supported by the platform"),
+	ECVF_RenderThreadSafe
+);
+
 DEFINE_LOG_CATEGORY_STATIC(LogClearQuad, Log, Log)
 
 static void ClearQuadSetup( FRHICommandList& RHICmdList, bool bClearColor, int32 NumClearColors, const FLinearColor* ClearColorArray, bool bClearDepth, float Depth, bool bClearStencil, uint32 Stencil, TFunction<void(FGraphicsPipelineStateInitializer&)> PSOModifier = nullptr)
@@ -121,7 +128,6 @@ static void ClearQuadSetup( FRHICommandList& RHICmdList, bool bClearColor, int32
 	PixelShader->SetColors(RHICmdList, ClearColorArray, NumClearColors);
 }
 
-const uint32 GMaxSizeUAVDMA = 0;
 static void ClearUAVShader(FRHICommandList& RHICmdList, FUnorderedAccessViewRHIParamRef UnorderedAccessViewRHI, uint32 SizeInBytes, uint32 ClearValue)
 {
 	UE_CLOG((SizeInBytes & 0x3) != 0, LogClearQuad, Warning,
@@ -132,7 +138,6 @@ static void ClearUAVShader(FRHICommandList& RHICmdList, FUnorderedAccessViewRHIP
 	
 	uint32 NumDWordsToClear = (SizeInBytes + 3) / 4;
 	uint32 NumThreadGroupsX = (NumDWordsToClear + 63) / 64;
-
 	RHICmdList.SetComputeShader(ShaderRHI);
 	ComputeShader->SetParameters(RHICmdList, UnorderedAccessViewRHI, NumDWordsToClear, ClearValue);
 	RHICmdList.DispatchComputeShader(NumThreadGroupsX, 1, 1);
@@ -141,7 +146,8 @@ static void ClearUAVShader(FRHICommandList& RHICmdList, FUnorderedAccessViewRHIP
 
 void ClearUAV(FRHICommandList& RHICmdList, const FRWBufferStructured& StructuredBuffer, uint32 Value)
 {
-	if (StructuredBuffer.NumBytes <= GMaxSizeUAVDMA)
+	
+	if (StructuredBuffer.NumBytes <= uint32(CVarFastClearUAVMaxSize.GetValueOnRenderThread()))
 	{
 		uint32 Values[4] = { Value, Value, Value, Value };
 		RHICmdList.ClearTinyUAV(StructuredBuffer.UAV, Values);
@@ -154,7 +160,7 @@ void ClearUAV(FRHICommandList& RHICmdList, const FRWBufferStructured& Structured
 
 void ClearUAV(FRHICommandList& RHICmdList, const FRWBuffer& Buffer, uint32 Value)
 {
-	if (Buffer.NumBytes <= GMaxSizeUAVDMA)
+	if (Buffer.NumBytes <= uint32(CVarFastClearUAVMaxSize.GetValueOnRenderThread()))
 	{
 		uint32 Values[4] = { Value, Value, Value, Value };
 		RHICmdList.ClearTinyUAV(Buffer.UAV, Values);
@@ -167,7 +173,7 @@ void ClearUAV(FRHICommandList& RHICmdList, const FRWBuffer& Buffer, uint32 Value
 
 void ClearUAV(FRHICommandList& RHICmdList, FRHIUnorderedAccessView* Buffer, uint32 NumBytes, uint32 Value)
 {
-	if (NumBytes <= GMaxSizeUAVDMA)
+	if (NumBytes <= uint32(CVarFastClearUAVMaxSize.GetValueOnRenderThread()))
 	{
 		uint32 Values[4] = { Value, Value, Value, Value };
 		RHICmdList.ClearTinyUAV(Buffer, Values);
diff --git a/Engine/Source/Runtime/UtilityShaders/Public/ClearQuad.h b/Engine/Source/Runtime/UtilityShaders/Public/ClearQuad.h
--- a/Engine/Source/Runtime/UtilityShaders/Public/ClearQuad.h
+++ b/Engine/Source/Runtime/UtilityShaders/Public/ClearQuad.h
@@ -43,7 +43,6 @@ struct FClearQuadCallbacks
 	TFunction<void(FRHICommandList&)> PostClear = nullptr;
 };
 
-extern UTILITYSHADERS_API const uint32 GMaxSizeUAVDMA;
 extern UTILITYSHADERS_API void ClearUAV(FRHICommandList& RHICmdList, const FRWBufferStructured& StructuredBuffer, uint32 Value);
 extern UTILITYSHADERS_API void ClearUAV(FRHICommandList& RHICmdList, const FRWBuffer& Buffer, uint32 Value);
 extern UTILITYSHADERS_API void ClearUAV(FRHICommandList& RHICmdList, FRHIUnorderedAccessView* Buffer, uint32 NumBytes, uint32 Value);
