MediaShaders: ability to crop images during YUY conversion (UE-65049)

(Edigrating 4471235 from Dev-Console to Release-4.21)

diff --git a/Engine/Plugins/Media/WebMMedia/Source/WebMMedia/Private/Player/WebMVideoDecoder.cpp b/Engine/Plugins/Media/WebMMedia/Source/WebMMedia/Private/Player/WebMVideoDecoder.cpp
--- a/Engine/Plugins/Media/WebMMedia/Source/WebMMedia/Private/Player/WebMVideoDecoder.cpp
+++ b/Engine/Plugins/Media/WebMMedia/Source/WebMMedia/Private/Player/WebMVideoDecoder.cpp
@@ -140,7 +140,7 @@ void FWebMVideoDecoder::DoDecodeVideoFrames(const TArray<TSharedPtr<FWebMFrame>>
 
 			TSharedRef<FWebMMediaTextureSample, ESPMode::ThreadSafe> VideoSample = VideoSamplePool->AcquireShared();
 
-			VideoSample->Initialize(FIntPoint(Image->d_w, Image->d_h), FIntPoint(Image->stride[0], Image->d_h), VideoFrame->Time);
+			VideoSample->Initialize(FIntPoint(Image->d_w, Image->d_h), FIntPoint(Image->d_w, Image->d_h), VideoFrame->Time);
 
 			FConvertParams Params;
 			Params.VideoSample = VideoSample;
@@ -255,10 +255,10 @@ void FWebMVideoDecoder::ConvertYUVToRGBAndSubmit(const FConvertParams& Params)
 		}
 
 		SetGraphicsPipelineState(CommandList, GraphicsPSOInit);
-		PixelShader->SetParameters(CommandList, DecodedY->GetTexture2D(), DecodedU->GetTexture2D(), DecodedV->GetTexture2D(), MediaShaders::YuvToSrgbDefault, true);
+		PixelShader->SetParameters(CommandList, DecodedY->GetTexture2D(), DecodedU->GetTexture2D(), DecodedV->GetTexture2D(), FIntPoint(Image->d_w, Image->d_h), MediaShaders::YuvToSrgbDefault, true);
 
 		// draw full-size quad
-		CommandList.SetViewport(0, 0, 0.0f, Image->w, Image->d_h, 1.0f);
+		CommandList.SetViewport(0, 0, 0.0f, Image->d_w, Image->d_h, 1.0f);
 		CommandList.SetStreamSource(0, GMoviePlayerResources.VertexBufferRHI, 0);
 		CommandList.DrawPrimitive(PT_TriangleStrip, 0, 2, 1);
 		CommandList.CopyToResolveTarget(RenderTarget, RenderTarget, FResolveParams());
diff --git a/Engine/Shaders/Private/MediaShaders.usf b/Engine/Shaders/Private/MediaShaders.usf
--- a/Engine/Shaders/Private/MediaShaders.usf
+++ b/Engine/Shaders/Private/MediaShaders.usf
@@ -204,10 +204,12 @@ void YUVConvertPS(
 	float2 InUV : TEXCOORD0,
 	out float4 OutColor : SV_Target0)
 {
+	float2 ScaledUV = InUV * YUVConvertUB.UVScale;
+
 	float3 YUV;
-	YUV.x = YUVConvertUB.YTexture.Sample(YUVConvertUB.YSampler, InUV.xy).x;
-	YUV.y = YUVConvertUB.UTexture.Sample(YUVConvertUB.USampler, InUV.xy).x;
-	YUV.z = YUVConvertUB.VTexture.Sample(YUVConvertUB.VSampler, InUV.xy).x;
+	YUV.x = YUVConvertUB.YTexture.Sample(YUVConvertUB.YSampler, ScaledUV).x;
+	YUV.y = YUVConvertUB.UTexture.Sample(YUVConvertUB.USampler, ScaledUV).x;
+	YUV.z = YUVConvertUB.VTexture.Sample(YUVConvertUB.VSampler, ScaledUV).x;
 	
 	OutColor.xyz = YuvToRgb(YUV, YUVConvertUB.ColorTransform, YUVConvertUB.SrgbToLinear);
 	OutColor.w = 1.0f;
diff --git a/Engine/Source/Runtime/UtilityShaders/Private/MediaShaders.cpp b/Engine/Source/Runtime/UtilityShaders/Private/MediaShaders.cpp
--- a/Engine/Source/Runtime/UtilityShaders/Private/MediaShaders.cpp
+++ b/Engine/Source/Runtime/UtilityShaders/Private/MediaShaders.cpp
@@ -286,6 +286,7 @@ void FUYVYConvertPS::SetParameters(FRHICommandList& CommandList, TRefCountPtr<FR
 BEGIN_UNIFORM_BUFFER_STRUCT(FYUVConvertUB, )
 UNIFORM_MEMBER(FMatrix, ColorTransform)
 UNIFORM_MEMBER(uint32, SrgbToLinear)
+UNIFORM_MEMBER(FVector2D, UVScale)
 UNIFORM_MEMBER_TEXTURE(Texture2D, YTexture)
 UNIFORM_MEMBER_TEXTURE(Texture2D, UTexture)
 UNIFORM_MEMBER_TEXTURE(Texture2D, VTexture)
@@ -298,7 +299,7 @@ IMPLEMENT_UNIFORM_BUFFER_STRUCT(FYUVConvertUB, TEXT("YUVConvertUB"));
 IMPLEMENT_SHADER_TYPE(, FYUVConvertPS, TEXT("/Engine/Private/MediaShaders.usf"), TEXT("YUVConvertPS"), SF_Pixel);
 
 
-void FYUVConvertPS::SetParameters(FRHICommandList& CommandList, TRefCountPtr<FRHITexture2D> YTexture, TRefCountPtr<FRHITexture2D> UTexture, TRefCountPtr<FRHITexture2D> VTexture, const FMatrix& ColorTransform, bool SrgbToLinear)
+void FYUVConvertPS::SetParameters(FRHICommandList& CommandList, TRefCountPtr<FRHITexture2D> YTexture, TRefCountPtr<FRHITexture2D> UTexture, TRefCountPtr<FRHITexture2D> VTexture, const FIntPoint& OutputDimensions, const FMatrix& ColorTransform, bool SrgbToLinear)
 {
 	FYUVConvertUB UB;
 	{
@@ -310,6 +311,7 @@ void FYUVConvertPS::SetParameters(FRHICommandList& CommandList, TRefCountPtr<FRH
 		UB.YSampler = TStaticSamplerState<SF_Bilinear>::GetRHI();
 		UB.USampler = TStaticSamplerState<SF_Bilinear>::GetRHI();
 		UB.VSampler = TStaticSamplerState<SF_Bilinear>::GetRHI();
+		UB.UVScale = FVector2D((float) OutputDimensions.X / (float) YTexture->GetSizeX(), (float) OutputDimensions.Y / (float) YTexture->GetSizeY());
 	}
 
 	TUniformBufferRef<FYUVConvertUB> Data = TUniformBufferRef<FYUVConvertUB>::CreateUniformBufferImmediate(UB, UniformBuffer_SingleFrame);
diff --git a/Engine/Source/Runtime/UtilityShaders/Public/MediaShaders.h b/Engine/Source/Runtime/UtilityShaders/Public/MediaShaders.h
--- a/Engine/Source/Runtime/UtilityShaders/Public/MediaShaders.h
+++ b/Engine/Source/Runtime/UtilityShaders/Public/MediaShaders.h
@@ -420,7 +420,7 @@ public:
 		return bShaderHasOutdatedParameters;
 	}
 
-	UTILITYSHADERS_API void SetParameters(FRHICommandList& RHICmdList, TRefCountPtr<FRHITexture2D> YTexture, TRefCountPtr<FRHITexture2D> UTexture, TRefCountPtr<FRHITexture2D> VTexture, const FMatrix& ColorTransform, bool SrgbToLinear);
+	UTILITYSHADERS_API void SetParameters(FRHICommandList& RHICmdList, TRefCountPtr<FRHITexture2D> YTexture, TRefCountPtr<FRHITexture2D> UTexture, TRefCountPtr<FRHITexture2D> VTexture, const FIntPoint& OutputDimensions, const FMatrix& ColorTransform, bool SrgbToLinear);
 };
 
 
