Merging //UE4/Dev-Rendering-Graph@4492585 to Dev-Rendering (//UE4/Dev-Rendering)

diff --git a/Engine/Source/Runtime/UtilityShaders/Private/MediaShaders.cpp b/Engine/Source/Runtime/UtilityShaders/Private/MediaShaders.cpp
--- a/Engine/Source/Runtime/UtilityShaders/Private/MediaShaders.cpp
+++ b/Engine/Source/Runtime/UtilityShaders/Private/MediaShaders.cpp
@@ -59,14 +59,14 @@ IMPLEMENT_SHADER_TYPE(, FMediaShadersVS, TEXT("/Engine/Private/MediaShaders.usf"
 /* FAYUVConvertPS shader
  *****************************************************************************/
 
-BEGIN_UNIFORM_BUFFER_STRUCT(FAYUVConvertUB, )
-UNIFORM_MEMBER(FMatrix, ColorTransform)
-UNIFORM_MEMBER(uint32, SrgbToLinear)
-UNIFORM_MEMBER_TEXTURE(Texture2D, Texture)
-UNIFORM_MEMBER_SAMPLER(SamplerState, Sampler)
-END_UNIFORM_BUFFER_STRUCT(FAYUVConvertUB)
-
-IMPLEMENT_UNIFORM_BUFFER_STRUCT(FAYUVConvertUB, TEXT("AYUVConvertUB"));
+BEGIN_GLOBAL_SHADER_PARAMETER_STRUCT(FAYUVConvertUB, )
+SHADER_PARAMETER(FMatrix, ColorTransform)
+SHADER_PARAMETER(uint32, SrgbToLinear)
+SHADER_PARAMETER_TEXTURE(Texture2D, Texture)
+SHADER_PARAMETER_SAMPLER(SamplerState, Sampler)
+END_GLOBAL_SHADER_PARAMETER_STRUCT()
+
+IMPLEMENT_GLOBAL_SHADER_PARAMETER_STRUCT(FAYUVConvertUB, "AYUVConvertUB");
 IMPLEMENT_SHADER_TYPE(, FAYUVConvertPS, TEXT("/Engine/Private/MediaShaders.usf"), TEXT("AYUVConvertPS"), SF_Pixel);
 
 
@@ -88,14 +88,14 @@ void FAYUVConvertPS::SetParameters(FRHICommandList& CommandList, TRefCountPtr<FR
 /* FRGBConvertPS shader
  *****************************************************************************/
 
-BEGIN_UNIFORM_BUFFER_STRUCT(FBMPConvertUB, )
-UNIFORM_MEMBER(uint32, SrgbToLinear)
-UNIFORM_MEMBER(FVector2D, UVScale)
-UNIFORM_MEMBER_TEXTURE(Texture2D, Texture)
-UNIFORM_MEMBER_SAMPLER(SamplerState, Sampler)
-END_UNIFORM_BUFFER_STRUCT(FBMPConvertUB)
+BEGIN_GLOBAL_SHADER_PARAMETER_STRUCT(FBMPConvertUB, )
+SHADER_PARAMETER(uint32, SrgbToLinear)
+SHADER_PARAMETER(FVector2D, UVScale)
+SHADER_PARAMETER_TEXTURE(Texture2D, Texture)
+SHADER_PARAMETER_SAMPLER(SamplerState, Sampler)
+END_GLOBAL_SHADER_PARAMETER_STRUCT()
 
-IMPLEMENT_UNIFORM_BUFFER_STRUCT(FBMPConvertUB, TEXT("BMPConvertUB"));
+IMPLEMENT_GLOBAL_SHADER_PARAMETER_STRUCT(FBMPConvertUB, "BMPConvertUB");
 IMPLEMENT_SHADER_TYPE(, FBMPConvertPS, TEXT("/Engine/Private/MediaShaders.usf"), TEXT("BMPConvertPS"), SF_Pixel);
 
 
@@ -117,17 +117,17 @@ void FBMPConvertPS::SetParameters(FRHICommandList& CommandList, TRefCountPtr<FRH
 /* FNV12ConvertPS shader
  *****************************************************************************/
 
-BEGIN_UNIFORM_BUFFER_STRUCT(FNV12ConvertUB, )
-UNIFORM_MEMBER(FMatrix, ColorTransform)
-UNIFORM_MEMBER(uint32, OutputWidth)
-UNIFORM_MEMBER(uint32, SrgbToLinear)
-UNIFORM_MEMBER(FVector2D, UVScale)
-UNIFORM_MEMBER_TEXTURE(Texture2D, Texture)
-UNIFORM_MEMBER_SAMPLER(SamplerState, SamplerB)
-UNIFORM_MEMBER_SAMPLER(SamplerState, SamplerP)
-END_UNIFORM_BUFFER_STRUCT(FNV12ConvertUB)
-
-IMPLEMENT_UNIFORM_BUFFER_STRUCT(FNV12ConvertUB, TEXT("NV12ConvertUB"));
+BEGIN_GLOBAL_SHADER_PARAMETER_STRUCT(FNV12ConvertUB, )
+SHADER_PARAMETER(FMatrix, ColorTransform)
+SHADER_PARAMETER(uint32, OutputWidth)
+SHADER_PARAMETER(uint32, SrgbToLinear)
+SHADER_PARAMETER(FVector2D, UVScale)
+SHADER_PARAMETER_TEXTURE(Texture2D, Texture)
+SHADER_PARAMETER_SAMPLER(SamplerState, SamplerB)
+SHADER_PARAMETER_SAMPLER(SamplerState, SamplerP)
+END_GLOBAL_SHADER_PARAMETER_STRUCT()
+
+IMPLEMENT_GLOBAL_SHADER_PARAMETER_STRUCT(FNV12ConvertUB, "NV12ConvertUB");
 IMPLEMENT_SHADER_TYPE(, FNV12ConvertPS, TEXT("/Engine/Private/MediaShaders.usf"), TEXT("NV12ConvertPS"), SF_Pixel);
 
 
@@ -152,17 +152,17 @@ void FNV12ConvertPS::SetParameters(FRHICommandList& CommandList, TRefCountPtr<FR
 /* FNV21ConvertPS shader
  *****************************************************************************/
 
-BEGIN_UNIFORM_BUFFER_STRUCT(FNV21ConvertUB, )
-UNIFORM_MEMBER(FMatrix, ColorTransform)
-UNIFORM_MEMBER(uint32, OutputWidth)
-UNIFORM_MEMBER(uint32, SrgbToLinear)
-UNIFORM_MEMBER(FVector2D, UVScale)
-UNIFORM_MEMBER_TEXTURE(Texture2D, Texture)
-UNIFORM_MEMBER_SAMPLER(SamplerState, SamplerB)
-UNIFORM_MEMBER_SAMPLER(SamplerState, SamplerP)
-END_UNIFORM_BUFFER_STRUCT(FNV21ConvertUB)
-
-IMPLEMENT_UNIFORM_BUFFER_STRUCT(FNV21ConvertUB, TEXT("NV21ConvertUB"));
+BEGIN_GLOBAL_SHADER_PARAMETER_STRUCT(FNV21ConvertUB, )
+SHADER_PARAMETER(FMatrix, ColorTransform)
+SHADER_PARAMETER(uint32, OutputWidth)
+SHADER_PARAMETER(uint32, SrgbToLinear)
+SHADER_PARAMETER(FVector2D, UVScale)
+SHADER_PARAMETER_TEXTURE(Texture2D, Texture)
+SHADER_PARAMETER_SAMPLER(SamplerState, SamplerB)
+SHADER_PARAMETER_SAMPLER(SamplerState, SamplerP)
+END_GLOBAL_SHADER_PARAMETER_STRUCT()
+
+IMPLEMENT_GLOBAL_SHADER_PARAMETER_STRUCT(FNV21ConvertUB, "NV21ConvertUB");
 IMPLEMENT_SHADER_TYPE(, FNV21ConvertPS, TEXT("/Engine/Private/MediaShaders.usf"), TEXT("NV21ConvertPS"), SF_Pixel);
 
 
@@ -187,14 +187,14 @@ void FNV21ConvertPS::SetParameters(FRHICommandList& CommandList, TRefCountPtr<FR
 /* FRGBConvertPS shader
  *****************************************************************************/
 
-BEGIN_UNIFORM_BUFFER_STRUCT(FRGBConvertUB, )
-UNIFORM_MEMBER(FVector2D, UVScale)
-UNIFORM_MEMBER(uint32, SrgbToLinear)
-UNIFORM_MEMBER_TEXTURE(Texture2D, Texture)
-UNIFORM_MEMBER_SAMPLER(SamplerState, Sampler)
-END_UNIFORM_BUFFER_STRUCT(FRGBConvertUB)
+BEGIN_GLOBAL_SHADER_PARAMETER_STRUCT(FRGBConvertUB, )
+SHADER_PARAMETER(FVector2D, UVScale)
+SHADER_PARAMETER(uint32, SrgbToLinear)
+SHADER_PARAMETER_TEXTURE(Texture2D, Texture)
+SHADER_PARAMETER_SAMPLER(SamplerState, Sampler)
+END_GLOBAL_SHADER_PARAMETER_STRUCT()
 
-IMPLEMENT_UNIFORM_BUFFER_STRUCT(FRGBConvertUB, TEXT("RGBConvertUB"));
+IMPLEMENT_GLOBAL_SHADER_PARAMETER_STRUCT(FRGBConvertUB, "RGBConvertUB");
 IMPLEMENT_SHADER_TYPE(, FRGBConvertPS, TEXT("/Engine/Private/MediaShaders.usf"), TEXT("RGBConvertPS"), SF_Pixel);
 
 
@@ -216,16 +216,16 @@ void FRGBConvertPS::SetParameters(FRHICommandList& CommandList, TRefCountPtr<FRH
 /* FYCbCrConvertUB shader
  *****************************************************************************/
 
-BEGIN_UNIFORM_BUFFER_STRUCT(FYCbCrConvertUB, )
-UNIFORM_MEMBER(FMatrix, ColorTransform)
-UNIFORM_MEMBER(uint32, SrgbToLinear)
-UNIFORM_MEMBER_TEXTURE(Texture2D, LumaTexture)
-UNIFORM_MEMBER_TEXTURE(Texture2D, CbCrTexture)
-UNIFORM_MEMBER_SAMPLER(SamplerState, LumaSampler)
-UNIFORM_MEMBER_SAMPLER(SamplerState, CbCrSampler)
-END_UNIFORM_BUFFER_STRUCT(FYCbCrConvertUB)
+BEGIN_GLOBAL_SHADER_PARAMETER_STRUCT(FYCbCrConvertUB, )
+SHADER_PARAMETER(FMatrix, ColorTransform)
+SHADER_PARAMETER(uint32, SrgbToLinear)
+SHADER_PARAMETER_TEXTURE(Texture2D, LumaTexture)
+SHADER_PARAMETER_TEXTURE(Texture2D, CbCrTexture)
+SHADER_PARAMETER_SAMPLER(SamplerState, LumaSampler)
+SHADER_PARAMETER_SAMPLER(SamplerState, CbCrSampler)
+END_GLOBAL_SHADER_PARAMETER_STRUCT()
 
-IMPLEMENT_UNIFORM_BUFFER_STRUCT(FYCbCrConvertUB, TEXT("YCbCrConvertUB"));
+IMPLEMENT_GLOBAL_SHADER_PARAMETER_STRUCT(FYCbCrConvertUB, "YCbCrConvertUB");
 IMPLEMENT_SHADER_TYPE(, FYCbCrConvertPS, TEXT("/Engine/Private/MediaShaders.usf"), TEXT("YCbCrConvertPS"), SF_Pixel);
 IMPLEMENT_SHADER_TYPE(, FYCbCrConvertPS_4x4Matrix, TEXT("/Engine/Private/MediaShaders.usf"), TEXT("YCbCrConvertPS_4x4Matrix"), SF_Pixel);
 
@@ -250,16 +250,16 @@ void FYCbCrConvertPS::SetParameters(FRHICommandList& CommandList, TRefCountPtr<F
 /* FUYVYConvertPS shader
  *****************************************************************************/
 
-BEGIN_UNIFORM_BUFFER_STRUCT(FUYVYConvertUB, )
-UNIFORM_MEMBER(FMatrix, ColorTransform)
-UNIFORM_MEMBER(uint32, SrgbToLinear)
-UNIFORM_MEMBER(uint32, Width)
-UNIFORM_MEMBER_TEXTURE(Texture2D, Texture)
-UNIFORM_MEMBER_SAMPLER(SamplerState, SamplerB)
-UNIFORM_MEMBER_SAMPLER(SamplerState, SamplerP)
-END_UNIFORM_BUFFER_STRUCT(FUYVYConvertUB)
+BEGIN_GLOBAL_SHADER_PARAMETER_STRUCT(FUYVYConvertUB, )
+SHADER_PARAMETER(FMatrix, ColorTransform)
+SHADER_PARAMETER(uint32, SrgbToLinear)
+SHADER_PARAMETER(uint32, Width)
+SHADER_PARAMETER_TEXTURE(Texture2D, Texture)
+SHADER_PARAMETER_SAMPLER(SamplerState, SamplerB)
+SHADER_PARAMETER_SAMPLER(SamplerState, SamplerP)
+END_GLOBAL_SHADER_PARAMETER_STRUCT()
 
-IMPLEMENT_UNIFORM_BUFFER_STRUCT(FUYVYConvertUB, TEXT("UYVYConvertUB"));
+IMPLEMENT_GLOBAL_SHADER_PARAMETER_STRUCT(FUYVYConvertUB, "UYVYConvertUB");
 IMPLEMENT_SHADER_TYPE(, FUYVYConvertPS, TEXT("/Engine/Private/MediaShaders.usf"), TEXT("UYVYConvertPS"), SF_Pixel);
 
 
@@ -283,18 +283,18 @@ void FUYVYConvertPS::SetParameters(FRHICommandList& CommandList, TRefCountPtr<FR
 /* FYUVConvertPS shader
  *****************************************************************************/
 
-BEGIN_UNIFORM_BUFFER_STRUCT(FYUVConvertUB, )
-UNIFORM_MEMBER(FMatrix, ColorTransform)
-UNIFORM_MEMBER(uint32, SrgbToLinear)
-UNIFORM_MEMBER_TEXTURE(Texture2D, YTexture)
-UNIFORM_MEMBER_TEXTURE(Texture2D, UTexture)
-UNIFORM_MEMBER_TEXTURE(Texture2D, VTexture)
-UNIFORM_MEMBER_SAMPLER(SamplerState, YSampler)
-UNIFORM_MEMBER_SAMPLER(SamplerState, USampler)
-UNIFORM_MEMBER_SAMPLER(SamplerState, VSampler)
-END_UNIFORM_BUFFER_STRUCT(FYUVConvertUB)
-
-IMPLEMENT_UNIFORM_BUFFER_STRUCT(FYUVConvertUB, TEXT("YUVConvertUB"));
+BEGIN_GLOBAL_SHADER_PARAMETER_STRUCT(FYUVConvertUB, )
+SHADER_PARAMETER(FMatrix, ColorTransform)
+SHADER_PARAMETER(uint32, SrgbToLinear)
+SHADER_PARAMETER_TEXTURE(Texture2D, YTexture)
+SHADER_PARAMETER_TEXTURE(Texture2D, UTexture)
+SHADER_PARAMETER_TEXTURE(Texture2D, VTexture)
+SHADER_PARAMETER_SAMPLER(SamplerState, YSampler)
+SHADER_PARAMETER_SAMPLER(SamplerState, USampler)
+SHADER_PARAMETER_SAMPLER(SamplerState, VSampler)
+END_GLOBAL_SHADER_PARAMETER_STRUCT()
+
+IMPLEMENT_GLOBAL_SHADER_PARAMETER_STRUCT(FYUVConvertUB, "YUVConvertUB");
 IMPLEMENT_SHADER_TYPE(, FYUVConvertPS, TEXT("/Engine/Private/MediaShaders.usf"), TEXT("YUVConvertPS"), SF_Pixel);
 
 
@@ -320,17 +320,17 @@ void FYUVConvertPS::SetParameters(FRHICommandList& CommandList, TRefCountPtr<FRH
 /* FYUY2ConvertPS shader
  *****************************************************************************/
 
-BEGIN_UNIFORM_BUFFER_STRUCT(FYUY2ConvertUB, )
-UNIFORM_MEMBER(FMatrix, ColorTransform)
-UNIFORM_MEMBER(uint32, OutputWidth)
-UNIFORM_MEMBER(uint32, SrgbToLinear)
-UNIFORM_MEMBER(FVector2D, UVScale)
-UNIFORM_MEMBER_TEXTURE(Texture2D, Texture)
-UNIFORM_MEMBER_SAMPLER(SamplerState, SamplerB)
-UNIFORM_MEMBER_SAMPLER(SamplerState, SamplerP)
-END_UNIFORM_BUFFER_STRUCT(FYUY2ConvertUB)
-
-IMPLEMENT_UNIFORM_BUFFER_STRUCT(FYUY2ConvertUB, TEXT("YUY2ConvertUB"));
+BEGIN_GLOBAL_SHADER_PARAMETER_STRUCT(FYUY2ConvertUB, )
+SHADER_PARAMETER(FMatrix, ColorTransform)
+SHADER_PARAMETER(uint32, OutputWidth)
+SHADER_PARAMETER(uint32, SrgbToLinear)
+SHADER_PARAMETER(FVector2D, UVScale)
+SHADER_PARAMETER_TEXTURE(Texture2D, Texture)
+SHADER_PARAMETER_SAMPLER(SamplerState, SamplerB)
+SHADER_PARAMETER_SAMPLER(SamplerState, SamplerP)
+END_GLOBAL_SHADER_PARAMETER_STRUCT()
+
+IMPLEMENT_GLOBAL_SHADER_PARAMETER_STRUCT(FYUY2ConvertUB, "YUY2ConvertUB");
 IMPLEMENT_SHADER_TYPE(, FYUY2ConvertPS, TEXT("/Engine/Private/MediaShaders.usf"), TEXT("YUY2ConvertPS"), SF_Pixel);
 
 
@@ -355,16 +355,16 @@ void FYUY2ConvertPS::SetParameters(FRHICommandList& CommandList, TRefCountPtr<FR
 /* FYVYUConvertPS shader
  *****************************************************************************/
 
-BEGIN_UNIFORM_BUFFER_STRUCT(FYVYUConvertUB, )
-UNIFORM_MEMBER(FMatrix, ColorTransform)
-UNIFORM_MEMBER(uint32, SrgbToLinear)
-UNIFORM_MEMBER(uint32, Width)
-UNIFORM_MEMBER_TEXTURE(Texture2D, Texture)
-UNIFORM_MEMBER_SAMPLER(SamplerState, SamplerB)
-UNIFORM_MEMBER_SAMPLER(SamplerState, SamplerP)
-END_UNIFORM_BUFFER_STRUCT(FYVYUConvertUB)
+BEGIN_GLOBAL_SHADER_PARAMETER_STRUCT(FYVYUConvertUB, )
+SHADER_PARAMETER(FMatrix, ColorTransform)
+SHADER_PARAMETER(uint32, SrgbToLinear)
+SHADER_PARAMETER(uint32, Width)
+SHADER_PARAMETER_TEXTURE(Texture2D, Texture)
+SHADER_PARAMETER_SAMPLER(SamplerState, SamplerB)
+SHADER_PARAMETER_SAMPLER(SamplerState, SamplerP)
+END_GLOBAL_SHADER_PARAMETER_STRUCT()
 
-IMPLEMENT_UNIFORM_BUFFER_STRUCT(FYVYUConvertUB, TEXT("YVYUConvertUB"));
+IMPLEMENT_GLOBAL_SHADER_PARAMETER_STRUCT(FYVYUConvertUB, "YVYUConvertUB");
 IMPLEMENT_SHADER_TYPE(, FYVYUConvertPS, TEXT("/Engine/Private/MediaShaders.usf"), TEXT("YVYUConvertPS"), SF_Pixel);
 
 
diff --git a/Engine/Source/Runtime/UtilityShaders/Private/OneColorShader.cpp b/Engine/Source/Runtime/UtilityShaders/Private/OneColorShader.cpp
--- a/Engine/Source/Runtime/UtilityShaders/Private/OneColorShader.cpp
+++ b/Engine/Source/Runtime/UtilityShaders/Private/OneColorShader.cpp
@@ -3,11 +3,11 @@
 
 #include "OneColorShader.h"
 
-BEGIN_UNIFORM_BUFFER_STRUCT(FClearShaderUB, )
-	UNIFORM_MEMBER_ARRAY(FVector4, DrawColorMRT, [8] )
-END_UNIFORM_BUFFER_STRUCT(FClearShaderUB)
+BEGIN_GLOBAL_SHADER_PARAMETER_STRUCT(FClearShaderUB, )
+	SHADER_PARAMETER_ARRAY(FVector4, DrawColorMRT, [8] )
+END_GLOBAL_SHADER_PARAMETER_STRUCT()
 
-IMPLEMENT_UNIFORM_BUFFER_STRUCT(FClearShaderUB, TEXT("ClearShaderUB"));
+IMPLEMENT_GLOBAL_SHADER_PARAMETER_STRUCT(FClearShaderUB, "ClearShaderUB");
 
 void FOneColorPS::SetColors(FRHICommandList& RHICmdList, const FLinearColor* Colors, int32 NumColors)
 {
diff --git a/Engine/Source/Runtime/Windows/D3D11RHI/Private/D3D11UniformBuffer.cpp b/Engine/Source/Runtime/Windows/D3D11RHI/Private/D3D11UniformBuffer.cpp
--- a/Engine/Source/Runtime/Windows/D3D11RHI/Private/D3D11UniformBuffer.cpp
+++ b/Engine/Source/Runtime/Windows/D3D11RHI/Private/D3D11UniformBuffer.cpp
@@ -213,7 +213,7 @@ FUniformBufferRHIRef FD3D11DynamicRHI::RHICreateUniformBuffer(const void* Conten
 			FRHIResource* Resource = *(FRHIResource**)((uint8*)Contents + Layout.ResourceOffsets[i]);
 
 			// Allow null SRV's in uniform buffers for feature levels that don't support SRV's in shaders
-			if (!(GMaxRHIFeatureLevel <= ERHIFeatureLevel::ES3_1 && Layout.Resources[i] == UBMT_SRV))
+			if (!(GMaxRHIFeatureLevel <= ERHIFeatureLevel::ES3_1 && (Layout.Resources[i] == UBMT_SRV || Layout.Resources[i] == UBMT_GRAPH_TRACKED_SRV)))
 			{
 				checkf(Resource, TEXT("Invalid resource entry creating uniform buffer, %s.Resources[%u], ResourceType 0x%x."), *Layout.GetDebugName().ToString(), i, Layout.Resources[i]);
 			}
