Renderpass updates.

LongGPUTask. Self contained.

diff --git a/Engine/Source/Runtime/UtilityShaders/Private/LongGPUTask.cpp b/Engine/Source/Runtime/UtilityShaders/Private/LongGPUTask.cpp
--- a/Engine/Source/Runtime/UtilityShaders/Private/LongGPUTask.cpp
+++ b/Engine/Source/Runtime/UtilityShaders/Private/LongGPUTask.cpp
@@ -26,55 +26,61 @@ void IssueScalableLongGPUTask(FRHICommandListImmediate& RHICmdList, int32 NumIte
 
 	FGraphicsPipelineStateInitializer GraphicsPSOInit;
 
-	SetRenderTarget(RHICmdList, LongTaskRenderTarget, FTextureRHIRef(), true);
-	RHICmdList.ApplyCachedRenderTargets(GraphicsPSOInit);
+	// #todo-renderpasses DontLoad the correct action here?
+	FRHIRenderPassInfo RPInfo(LongTaskRenderTarget, ERenderTargetActions::DontLoad_Store);
+	TransitionRenderPassTargets(RHICmdList, RPInfo);
+	RHICmdList.BeginRenderPass(RPInfo, TEXT("LongGPUTask"));
+	{
+		RHICmdList.ApplyCachedRenderTargets(GraphicsPSOInit);
 
-	GraphicsPSOInit.BlendState = TStaticBlendState<>::GetRHI();
-	GraphicsPSOInit.RasterizerState = TStaticRasterizerState<>::GetRHI();
-	GraphicsPSOInit.DepthStencilState = TStaticDepthStencilState<false, CF_Always>::GetRHI();
+		GraphicsPSOInit.BlendState = TStaticBlendState<>::GetRHI();
+		GraphicsPSOInit.RasterizerState = TStaticRasterizerState<>::GetRHI();
+		GraphicsPSOInit.DepthStencilState = TStaticDepthStencilState<false, CF_Always>::GetRHI();
 
-	TShaderMapRef<TOneColorVS<true>> VertexShader(ShaderMap);
-	TShaderMapRef<FLongGPUTaskPS> PixelShader(ShaderMap);
+		TShaderMapRef<TOneColorVS<true>> VertexShader(ShaderMap);
+		TShaderMapRef<FLongGPUTaskPS> PixelShader(ShaderMap);
 
-	GraphicsPSOInit.BoundShaderState.VertexDeclarationRHI = GetVertexDeclarationFVector4();
-	GraphicsPSOInit.BoundShaderState.VertexShaderRHI = VertexShader->GetVertexShader();
-	GraphicsPSOInit.BoundShaderState.PixelShaderRHI = PixelShader->GetPixelShader();
-	GraphicsPSOInit.PrimitiveType = PT_TriangleStrip;
+		GraphicsPSOInit.BoundShaderState.VertexDeclarationRHI = GetVertexDeclarationFVector4();
+		GraphicsPSOInit.BoundShaderState.VertexShaderRHI = VertexShader->GetVertexShader();
+		GraphicsPSOInit.BoundShaderState.PixelShaderRHI = PixelShader->GetPixelShader();
+		GraphicsPSOInit.PrimitiveType = PT_TriangleStrip;
 
-	SetGraphicsPipelineState(RHICmdList, GraphicsPSOInit);
-	VertexShader->SetDepthParameter(RHICmdList, 0.0f);
+		SetGraphicsPipelineState(RHICmdList, GraphicsPSOInit);
+		VertexShader->SetDepthParameter(RHICmdList, 0.0f);
 
-	RHICmdList.SetStreamSource(0, GClearVertexBuffer.VertexBufferRHI, 0);
-	if (NumIteration == -1)
-	{
-		// Use the measured number of iterations to achieve 100ms
-		// If the query results are still not available, stall
-		if (NumMeasuredIterationsToAchieve100ms == 0)
+		RHICmdList.SetStreamSource(0, GClearVertexBuffer.VertexBufferRHI, 0);
+		if (NumIteration == -1)
 		{
-			if (TimeQueryStart != nullptr && TimeQueryEnd != nullptr) // Not all platforms/drivers support RQT_AbsoluteTime queries
-			{
-				uint64 StartTime = 0;
-				uint64 EndTime = 0;
-
-				// Results are in microseconds
-				RHICmdList.GetRenderQueryResult(TimeQueryStart, StartTime, true);
-				RHICmdList.GetRenderQueryResult(TimeQueryEnd, EndTime, true);
-
-				NumMeasuredIterationsToAchieve100ms = FMath::Clamp(FMath::FloorToInt(100.0f / ((EndTime - StartTime) / 1000.0f / NumIterationsForMeasurement)), 1, 200);
-			}
-			else
+			// Use the measured number of iterations to achieve 100ms
+			// If the query results are still not available, stall
+			if (NumMeasuredIterationsToAchieve100ms == 0)
 			{
-				NumMeasuredIterationsToAchieve100ms = 5; // Use a constant time on these platforms
+				if (TimeQueryStart != nullptr && TimeQueryEnd != nullptr) // Not all platforms/drivers support RQT_AbsoluteTime queries
+				{
+					uint64 StartTime = 0;
+					uint64 EndTime = 0;
+
+					// Results are in microseconds
+					RHICmdList.GetRenderQueryResult(TimeQueryStart, StartTime, true);
+					RHICmdList.GetRenderQueryResult(TimeQueryEnd, EndTime, true);
+
+					NumMeasuredIterationsToAchieve100ms = FMath::Clamp(FMath::FloorToInt(100.0f / ((EndTime - StartTime) / 1000.0f / NumIterationsForMeasurement)), 1, 200);
+				}
+				else
+				{
+					NumMeasuredIterationsToAchieve100ms = 5; // Use a constant time on these platforms
+				}
 			}
-		}
 
-		NumIteration = NumMeasuredIterationsToAchieve100ms;
-	}
+			NumIteration = NumMeasuredIterationsToAchieve100ms;
+		}
 
-	for (int32 Iteration = 0; Iteration < NumIteration; Iteration++)
-	{
-		RHICmdList.DrawPrimitive(0, 2, 1);
+		for (int32 Iteration = 0; Iteration < NumIteration; Iteration++)
+		{
+			RHICmdList.DrawPrimitive(0, 2, 1);
+		}
 	}
+	RHICmdList.EndRenderPass();
 }
 
 void MeasureLongGPUTaskExecutionTime(FRHICommandListImmediate& RHICmdList)
