Fixed built-in ray tracing shader compilation issue that results in crash at startup if RT is disabled for the project, but supported and enabled at the RHI level.

diff --git a/Engine/Source/Runtime/UtilityShaders/Public/BuiltInRayTracingShaders.h b/Engine/Source/Runtime/UtilityShaders/Public/BuiltInRayTracingShaders.h
--- a/Engine/Source/Runtime/UtilityShaders/Public/BuiltInRayTracingShaders.h
+++ b/Engine/Source/Runtime/UtilityShaders/Public/BuiltInRayTracingShaders.h
@@ -17,7 +17,8 @@ protected:
 public:
 	static bool ShouldCompilePermutation(const FGlobalShaderPermutationParameters& Parameters)
 	{
-		return ShouldCompileRayTracingShadersForProject(Parameters.Platform);
+		// Built-in ray tracing shaders are always compiled for RHIs that support them, regardless of whether RT is enabled for the project.
+		return RHISupportsRayTracingShaders(Parameters.Platform);
 	}
 
 	FBuiltInRayTracingShader() = default;
