Reimplement ScalableLongGPUTask. Static variables are removed so no more hidden states

diff --git a/Engine/Source/Runtime/Launch/Private/LaunchEngineLoop.cpp b/Engine/Source/Runtime/Launch/Private/LaunchEngineLoop.cpp
--- a/Engine/Source/Runtime/Launch/Private/LaunchEngineLoop.cpp
+++ b/Engine/Source/Runtime/Launch/Private/LaunchEngineLoop.cpp
@@ -3385,10 +3385,6 @@ void FEngineLoop::Exit()
 	MALLOC_PROFILER( GMalloc->Exec(nullptr, TEXT("MPROF STOP"), *GLog);	);
 #endif // !ANDROID
 
-#if WITH_PROFILEGPU
-	ClearLongGPUTaskQueries();
-#endif
-
 	// Stop the rendering thread.
 	StopRenderingThread();
 
diff --git a/Engine/Source/Runtime/UtilityShaders/Private/LongGPUTask.cpp b/Engine/Source/Runtime/UtilityShaders/Private/LongGPUTask.cpp
--- a/Engine/Source/Runtime/UtilityShaders/Private/LongGPUTask.cpp
+++ b/Engine/Source/Runtime/UtilityShaders/Private/LongGPUTask.cpp
@@ -9,34 +9,26 @@
 
 IMPLEMENT_SHADER_TYPE(, FLongGPUTaskPS, TEXT("/Engine/Private/OneColorShader.usf"), TEXT("MainLongGPUTask"), SF_Pixel);
 
-int32 NumMeasuredIterationsToAchieve100ms = 0;
-
-const int32 NumIterationsForMeasurement = 5;
-
-static bool HasRun = false;
-FRenderQueryRHIRef TimeQueryStart;
-FRenderQueryRHIRef TimeQueryEnd;
+int32 NumMeasuredIterationsToAchieve500ms = 0;
 
 void IssueScalableLongGPUTask(FRHICommandListImmediate& RHICmdList, int32 NumIteration /* = -1 by default */)
 {
 	FRHIResourceCreateInfo Info;
 	FTexture2DRHIRef LongTaskRenderTarget = RHICreateTexture2D(1920, 1080, PF_B8G8R8A8, 1, 1, TexCreate_RenderTargetable, Info);
 
-	auto ShaderMap = GetGlobalShaderMap(GMaxRHIFeatureLevel);
-
-	FGraphicsPipelineStateInitializer GraphicsPSOInit;
-
-	// #todo-renderpasses DontLoad the correct action here?
 	FRHIRenderPassInfo RPInfo(LongTaskRenderTarget, ERenderTargetActions::DontLoad_Store);
 	TransitionRenderPassTargets(RHICmdList, RPInfo);
 	RHICmdList.BeginRenderPass(RPInfo, TEXT("LongGPUTask"));
+
 	{
+		FGraphicsPipelineStateInitializer GraphicsPSOInit;
 		RHICmdList.ApplyCachedRenderTargets(GraphicsPSOInit);
 
 		GraphicsPSOInit.BlendState = TStaticBlendState<>::GetRHI();
 		GraphicsPSOInit.RasterizerState = TStaticRasterizerState<>::GetRHI();
 		GraphicsPSOInit.DepthStencilState = TStaticDepthStencilState<false, CF_Always>::GetRHI();
 
+		auto ShaderMap = GetGlobalShaderMap(GMaxRHIFeatureLevel);
 		TShaderMapRef<TOneColorVS<true>> VertexShader(ShaderMap);
 		TShaderMapRef<FLongGPUTaskPS> PixelShader(ShaderMap);
 
@@ -49,30 +41,10 @@ void IssueScalableLongGPUTask(FRHICommandListImmediate& RHICmdList, int32 NumIte
 		VertexShader->SetDepthParameter(RHICmdList, 0.0f);
 
 		RHICmdList.SetStreamSource(0, GClearVertexBuffer.VertexBufferRHI, 0);
+
 		if (NumIteration == -1)
 		{
-			// Use the measured number of iterations to achieve 100ms
-			// If the query results are still not available, stall
-			if (NumMeasuredIterationsToAchieve100ms == 0)
-			{
-				if (TimeQueryStart != nullptr && TimeQueryEnd != nullptr) // Not all platforms/drivers support RQT_AbsoluteTime queries
-				{
-					uint64 StartTime = 0;
-					uint64 EndTime = 0;
-
-					// Results are in microseconds
-					RHICmdList.GetRenderQueryResult(TimeQueryStart, StartTime, true);
-					RHICmdList.GetRenderQueryResult(TimeQueryEnd, EndTime, true);
-
-					NumMeasuredIterationsToAchieve100ms = FMath::Clamp(FMath::FloorToInt(100.0f / ((EndTime - StartTime) / 1000.0f / NumIterationsForMeasurement)), 1, 200);
-				}
-				else
-				{
-					NumMeasuredIterationsToAchieve100ms = 5; // Use a constant time on these platforms
-				}
-			}
-
-			NumIteration = NumMeasuredIterationsToAchieve100ms;
+			NumIteration = NumMeasuredIterationsToAchieve500ms;
 		}
 
 		for (int32 Iteration = 0; Iteration < NumIteration; Iteration++)
@@ -80,33 +52,49 @@ void IssueScalableLongGPUTask(FRHICommandListImmediate& RHICmdList, int32 NumIte
 			RHICmdList.DrawPrimitive(0, 2, 1);
 		}
 	}
+
 	RHICmdList.EndRenderPass();
 }
 
 void MeasureLongGPUTaskExecutionTime(FRHICommandListImmediate& RHICmdList)
 {
-	if (HasRun) { return; }
+	const int32 NumIterationsForMeasurement = 5;
 
-	check(TimeQueryStart == nullptr && TimeQueryEnd == nullptr);
+	FRenderQueryRHIRef TimeQueryStart;
+	FRenderQueryRHIRef TimeQueryEnd;
 
-	TimeQueryStart = RHICmdList.CreateRenderQuery(RQT_AbsoluteTime);
-	TimeQueryEnd = RHICmdList.CreateRenderQuery(RQT_AbsoluteTime);
+	TimeQueryStart = RHICreateRenderQuery(RQT_AbsoluteTime);
+	TimeQueryEnd = RHICreateRenderQuery(RQT_AbsoluteTime);
 
-	if (TimeQueryStart != nullptr && TimeQueryEnd != nullptr) // Not all platforms/drivers support RQT_AbsoluteTime queries
+	if (TimeQueryStart == nullptr || TimeQueryEnd == nullptr)
 	{
-		RHICmdList.EndRenderQuery(TimeQueryStart);
+		// Not all platforms/drivers support RQT_AbsoluteTime queries
+		// Use fixed number of iterations on those platforms
+		NumMeasuredIterationsToAchieve500ms = 5;
+		return;
+	}
 
-		IssueScalableLongGPUTask(RHICmdList, NumIterationsForMeasurement);
+	RHICmdList.EndRenderQuery(TimeQueryStart);
 
-		RHICmdList.EndRenderQuery(TimeQueryEnd);
-		HasRun = true;
-	}
+	IssueScalableLongGPUTask(RHICmdList, NumIterationsForMeasurement);
 
-	ClearLongGPUTaskQueries();
-}
+	RHICmdList.EndRenderQuery(TimeQueryEnd);
 
-void ClearLongGPUTaskQueries()
-{
-	TimeQueryStart = nullptr;
-	TimeQueryEnd = nullptr;
+	// Required by DX12 to resolve the query
+	RHICmdList.SubmitCommandsHint();
+
+	uint64 StartTime = 0;
+	uint64 EndTime = 0;
+
+	if (RHIGetRenderQueryResult(TimeQueryStart, StartTime, true) && RHIGetRenderQueryResult(TimeQueryEnd, EndTime, true))
+	{
+		NumMeasuredIterationsToAchieve500ms = FMath::Clamp(FMath::FloorToInt(500.0f / ((EndTime - StartTime) / 1000.0f / NumIterationsForMeasurement)), 1, 200);
+	}
+	else
+	{
+		// Sometimes it fails even the platform supports RQT_AbsoluteTime
+		// Fallback and show a warning
+		NumMeasuredIterationsToAchieve500ms = 5;
+		UE_LOG(LogTemp, Display, TEXT("Unable to get render query result on a platform supporting RQT_AbsoluteTime queries, defaulting to %d iterations for LongGPUTask"), NumMeasuredIterationsToAchieve500ms);
+	}
 }
diff --git a/Engine/Source/Runtime/UtilityShaders/Public/LongGPUTask.h b/Engine/Source/Runtime/UtilityShaders/Public/LongGPUTask.h
--- a/Engine/Source/Runtime/UtilityShaders/Public/LongGPUTask.h
+++ b/Engine/Source/Runtime/UtilityShaders/Public/LongGPUTask.h
@@ -8,4 +8,3 @@ class FRHICommandListImmediate;
 
 extern UTILITYSHADERS_API void IssueScalableLongGPUTask(FRHICommandListImmediate& RHICmdList, int32 NumIteration = -1);
 extern UTILITYSHADERS_API void MeasureLongGPUTaskExecutionTime(FRHICommandListImmediate& RHICmdList);
-extern UTILITYSHADERS_API void ClearLongGPUTaskQueries();
