Render graph updates

diff --git a/Engine/Source/Runtime/RenderCore/Private/RenderGraphUtils.cpp b/Engine/Source/Runtime/RenderCore/Private/RenderGraphUtils.cpp
--- a/Engine/Source/Runtime/RenderCore/Private/RenderGraphUtils.cpp
+++ b/Engine/Source/Runtime/RenderCore/Private/RenderGraphUtils.cpp
@@ -1,9 +1,13 @@
 // Copyright 1998-2019 Epic Games, Inc. All Rights Reserved.
 
 #include "RenderGraphUtils.h"
+#include <initializer_list>
 
-
-void ClearUnusedGraphResourcesImpl(const FShaderParameterBindings& ShaderBindings, const FShaderParametersMetadata* ParametersMetadata, void* InoutParameters)
+void ClearUnusedGraphResourcesImpl(
+	const FShaderParameterBindings& ShaderBindings,
+	const FShaderParametersMetadata* ParametersMetadata,
+	void* InoutParameters,
+	std::initializer_list< FRDGResourceRef > ExcludeList)
 {
 	int32 GraphTextureId = 0;
 	int32 GraphSRVId = 0;
@@ -45,6 +49,15 @@ void ClearUnusedGraphResourcesImpl(const FShaderParameterBindings& ShaderBinding
 			continue;
 		}
 
+		for( FRDGResourceRef ExcludeResource : ExcludeList )
+		{
+			auto Resource = *reinterpret_cast<const FRDGResource* const*>(Base + ByteOffset);
+			if( Resource == ExcludeResource )
+			{
+				continue;
+			}
+		}
+
 		void** ResourcePointerAddress = reinterpret_cast<void**>(Base + ByteOffset);
 		*ResourcePointerAddress = nullptr;
 	}
diff --git a/Engine/Source/Runtime/RenderCore/Private/ShaderParameterStruct.cpp b/Engine/Source/Runtime/RenderCore/Private/ShaderParameterStruct.cpp
--- a/Engine/Source/Runtime/RenderCore/Private/ShaderParameterStruct.cpp
+++ b/Engine/Source/Runtime/RenderCore/Private/ShaderParameterStruct.cpp
@@ -53,10 +53,6 @@ struct FShaderParameterStructBindingContext
 				RenderTargetBindingSlotCppName = CppName;
 				continue;
 			}
-			else if (BaseType == UBMT_RDG_BUFFER)
-			{
-				continue;
-			}
 
 			// Compute the shader member name to look for according to nesting.
 			FString ShaderBindingName = FString::Printf(TEXT("%s%s"), MemberPrefix, Member.GetName());
@@ -78,6 +74,14 @@ struct FShaderParameterStructBindingContext
 				// The member name of a globally referenced struct is the not name on the struct.
 				ShaderBindingName = Member.GetStructMetadata()->GetShaderVariableName();
 			}
+			else if (BaseType == UBMT_RDG_BUFFER)
+			{
+				if( ParametersMap->ContainsParameterAllocation(*ShaderBindingName) )
+				{
+					UE_LOG(LogShaders, Fatal, TEXT("%s can't bind shader parameter %s as buffer. Use buffer SRV for reading in shader."), *CppName, *ShaderBindingName);
+				}
+				continue;
+			}
 			else if (bUseRootShaderParameters && (BaseType == UBMT_INT32 || BaseType == UBMT_UINT32 || BaseType == UBMT_FLOAT32))
 			{
 				// Constants are stored in the root shader parameter cbuffer when bUseRootShaderParameters == true.
diff --git a/Engine/Source/Runtime/RenderCore/Public/RenderGraphBuilder.h b/Engine/Source/Runtime/RenderCore/Public/RenderGraphBuilder.h
--- a/Engine/Source/Runtime/RenderCore/Public/RenderGraphBuilder.h
+++ b/Engine/Source/Runtime/RenderCore/Public/RenderGraphBuilder.h
@@ -326,6 +326,11 @@ public:
 		return UAV;
 	}
 
+	inline FRDGBufferUAVRef CreateUAV(FRDGBufferRef Buffer, EPixelFormat Format)
+	{
+		return CreateUAV( FRDGBufferUAVDesc(Buffer, Format) );
+	}
+
 	/** Allocates parameter struct specifically to survive through the life time of the render graph. */
 	template< typename ParameterStructType >
 	inline ParameterStructType* AllocParameters() const
diff --git a/Engine/Source/Runtime/RenderCore/Public/RenderGraphResources.h b/Engine/Source/Runtime/RenderCore/Public/RenderGraphResources.h
--- a/Engine/Source/Runtime/RenderCore/Public/RenderGraphResources.h
+++ b/Engine/Source/Runtime/RenderCore/Public/RenderGraphResources.h
@@ -21,6 +21,8 @@ class FRDGBuffer;
 class FRDGBufferSRV;
 class FRDGBufferUAV;
 
+class FRDGEventName;
+
 struct FPooledRDGBuffer;
 struct FRayTracingShaderBindingsWriter;
 
@@ -100,6 +102,10 @@ private:
 	template<typename TShaderClass>
 	friend void SetShaderParameters(FRayTracingShaderBindingsWriter& RTBindingsWriter, const TShaderClass* Shader, const typename TShaderClass::FParameters& Parameters);
 #endif
+
+	template< typename T >
+	friend void AddPass_ClearUAV(FRDGBuilder&, FRDGEventName&&, FRDGTextureRef, FRDGTextureUAVRef, const T(&ClearValues)[4]);
+	friend void AddPass_ClearUAV(FRDGBuilder&, FRDGEventName&&, FRDGBufferUAVRef, uint32);
 };
 
 /** Descriptor of a graph tracked texture. */
@@ -227,7 +233,7 @@ struct RENDERCORE_API FRDGBufferDesc
 	/** Bitfields describing the uses of that buffer. */
 	EBufferUsageFlags Usage = BUF_None;
 
-	/** The underlying RHI type to use. A bit of a work arround because RHI still have 3 different objects. */
+	/** The underlying RHI type to use. A bit of a work around because RHI still have 3 different objects. */
 	EUnderlyingType UnderlyingType = EUnderlyingType::VertexBuffer;
 
 	/** Returns the total number of bytes allocated for a such buffer. */
@@ -268,6 +274,26 @@ struct RENDERCORE_API FRDGBufferDesc
 		Desc.NumElements = NumElements;
 		return Desc;
 	}
+
+	static inline FRDGBufferDesc CreateIndirectDesc(uint32 NumElements = 1)
+	{
+		FRDGBufferDesc Desc;
+		Desc.UnderlyingType = EUnderlyingType::VertexBuffer;
+		Desc.Usage = EBufferUsageFlags(BUF_Static | BUF_DrawIndirect | BUF_UnorderedAccess | BUF_ShaderResource);
+		Desc.BytesPerElement = 4;
+		Desc.NumElements = NumElements;
+		return Desc;
+	}
+
+	static inline FRDGBufferDesc CreateStructuredDesc(uint32 BytesPerElement, uint32 NumElements)
+	{
+		FRDGBufferDesc Desc;
+		Desc.UnderlyingType = EUnderlyingType::StructuredBuffer;
+		Desc.Usage = EBufferUsageFlags(BUF_Static | BUF_UnorderedAccess | BUF_ShaderResource);
+		Desc.BytesPerElement = BytesPerElement;
+		Desc.NumElements = NumElements;
+		return Desc;
+	}
 };
 
 
@@ -275,26 +301,39 @@ struct RENDERCORE_API FRDGBufferSRVDesc
 {
 	FRDGBufferRef Buffer = nullptr;
 
-	/** Nymber of bytes per element (used for vertex buffer). */
+	/** Number of bytes per element (used for vertex buffer). */
 	uint32 BytesPerElement = 1;
 
 	/** Encoding format for the element (used for vertex buffer). */
 	EPixelFormat Format = PF_Unknown;
+
+	FRDGBufferSRVDesc() {}
+	FRDGBufferSRVDesc(FRDGBufferRef InBuffer);
+	FRDGBufferSRVDesc(FRDGBufferRef InBuffer, EPixelFormat InFormat)
+		: Buffer(InBuffer)
+		, Format(InFormat)
+	{
+		BytesPerElement = GPixelFormats[ Format ].BlockBytes;
+	}
 };
 
 struct RENDERCORE_API FRDGBufferUAVDesc
 {
 	FRDGBufferRef Buffer = nullptr;
 
-	/** Nymber of bytes per element (used for vertex buffer). */
+	/** Number of bytes per element (used for vertex buffer). */
 	EPixelFormat Format = PF_Unknown;
 
 	/** Whether the uav supports atomic counter or append buffer operations (used for structured buffers) */
 	bool bSupportsAtomicCounter = false;
 	bool bSupportsAppendBuffer = false;
 
-	/** Create descriptor for a the UAV of a buffer meant to be use for indirect dral call parameters. */
-	static inline FRDGBufferUAVDesc CreateIndirectDesc(FRDGBufferRef Buffer);
+	FRDGBufferUAVDesc() {}
+	FRDGBufferUAVDesc(FRDGBufferRef InBuffer);
+	FRDGBufferUAVDesc(FRDGBufferRef InBuffer, EPixelFormat InFormat)
+		: Buffer(InBuffer)
+		, Format(InFormat)
+	{}
 };
 
 
@@ -435,11 +474,29 @@ private:
 };
 
 
-FRDGBufferUAVDesc FRDGBufferUAVDesc::CreateIndirectDesc(FRDGBufferRef Buffer)
+inline FRDGBufferSRVDesc::FRDGBufferSRVDesc(FRDGBufferRef InBuffer)
+	: Buffer(InBuffer)
+{
+	if( Buffer->Desc.Usage & BUF_DrawIndirect )
+	{
+		BytesPerElement = 4;
+		Format = PF_R32_UINT;
+	}
+	else
+	{
+		checkf(Buffer->Desc.UnderlyingType != FRDGBufferDesc::EUnderlyingType::VertexBuffer, TEXT("VertexBuffer %s requires a type when creating a SRV."), Buffer->Name);
+	}
+}
+
+inline FRDGBufferUAVDesc::FRDGBufferUAVDesc(FRDGBufferRef InBuffer)
+	: Buffer(InBuffer)
 {
-	checkf(Buffer->Desc.Usage & BUF_DrawIndirect, TEXT("Buffer %s has not been created with a BUF_DrawIndirect flag."), Buffer->Name);
-	FRDGBufferUAVDesc Desc;
-	Desc.Buffer = Buffer;
-	Desc.Format = PF_R32_UINT;
-	return Desc;
+	if( Buffer->Desc.Usage & BUF_DrawIndirect )
+	{
+		Format = PF_R32_UINT;
+	}
+	else
+	{
+		checkf(Buffer->Desc.UnderlyingType != FRDGBufferDesc::EUnderlyingType::VertexBuffer, TEXT("VertexBuffer %s requires a type when creating a UAV."), Buffer->Name);
+	}
 }
diff --git a/Engine/Source/Runtime/RenderCore/Public/RenderGraphUtils.h b/Engine/Source/Runtime/RenderCore/Public/RenderGraphUtils.h
--- a/Engine/Source/Runtime/RenderCore/Public/RenderGraphUtils.h
+++ b/Engine/Source/Runtime/RenderCore/Public/RenderGraphUtils.h
@@ -25,11 +25,23 @@ BEGIN_SHADER_PARAMETER_STRUCT(FRenderTargetParameters, RENDERCORE_API)
 END_SHADER_PARAMETER_STRUCT()
 
 
-/** Clears all render graph tracked resources that are not used by a shader bindings. */
-extern RENDERCORE_API void ClearUnusedGraphResourcesImpl(const FShaderParameterBindings& ShaderBindings, const FShaderParametersMetadata* ParametersMetadata, void* InoutParameters);
+/**
+ * Clears all render graph tracked resources that are not bound by a shader.
+ * Excludes any resources on the ExcludeList from being cleared regardless of whether the 
+ * shader binds them or not. This is needed for resources that are used outside of shader
+ * bindings such as indirect arguments buffers.
+ */
+extern RENDERCORE_API void ClearUnusedGraphResourcesImpl(
+	const FShaderParameterBindings& ShaderBindings,
+	const FShaderParametersMetadata* ParametersMetadata,
+	void* InoutParameters,
+	std::initializer_list< FRDGResourceRef > ExcludeList);
 
 template<typename TShaderClass>
-FORCEINLINE_DEBUGGABLE void ClearUnusedGraphResources(const TShaderClass* Shader, typename TShaderClass::FParameters* InoutParameters)
+FORCEINLINE_DEBUGGABLE void ClearUnusedGraphResources(
+	const TShaderClass* Shader,
+	typename TShaderClass::FParameters* InoutParameters,
+	std::initializer_list< FRDGResourceRef > ExcludeList = {})
 {
 	const FShaderParametersMetadata* ParametersMetadata = TShaderClass::FParameters::FTypeInfo::GetStructMetadata();
 
@@ -39,7 +51,7 @@ FORCEINLINE_DEBUGGABLE void ClearUnusedGraphResources(const TShaderClass* Shader
 	ValidateShaderParameters(Shader, ParametersMetadata, InoutParameters);
 
 	// Clear the resources the shader won't need.
-	return ClearUnusedGraphResourcesImpl(Shader->Bindings, ParametersMetadata, InoutParameters);
+	return ClearUnusedGraphResourcesImpl(Shader->Bindings, ParametersMetadata, InoutParameters, ExcludeList);
 }
 
 
@@ -59,6 +71,13 @@ struct RENDERCORE_API FComputeShaderUtils
 	static constexpr int32 kGolden2DGroupSize = 8;
 
 	/** Compute the number of group to dispatch. */
+	static FIntVector GetGroupCount(const int32 ThreadCount, const int32 GroupSize)
+	{
+		return FIntVector(
+			FMath::DivideAndRoundUp(ThreadCount, GroupSize),
+			1,
+			1);
+	}
 	static FIntVector GetGroupCount(const FIntPoint& ThreadCount, const FIntPoint& GroupSize)
 	{
 		return FIntVector(
@@ -73,6 +92,13 @@ struct RENDERCORE_API FComputeShaderUtils
 			FMath::DivideAndRoundUp(ThreadCount.Y, GroupSize),
 			1);
 	}
+	static FIntVector GetGroupCount(const FIntVector& ThreadCount, const FIntVector& GroupSize)
+	{
+		return FIntVector(
+			FMath::DivideAndRoundUp(ThreadCount.X, GroupSize.X),
+			FMath::DivideAndRoundUp(ThreadCount.Y, GroupSize.Y),
+			FMath::DivideAndRoundUp(ThreadCount.Z, GroupSize.Z));
+	}
 
 
 	/** Dispatch a compute shader to rhi command list with its parameters. */
@@ -106,4 +132,30 @@ struct RENDERCORE_API FComputeShaderUtils
 			FComputeShaderUtils::Dispatch(RHICmdList, ComputeShader, *Parameters, GroupCount);
 		});
 	}
-};
+
+	/** Dispatch a compute shader to render graph builder with its parameters. */
+	template<typename TShaderClass>
+	static FORCEINLINE_DEBUGGABLE void AddPass(
+		FRDGBuilder& GraphBuilder,
+		FRDGEventName&& PassName,
+		const TShaderClass* ComputeShader,
+		typename TShaderClass::FParameters* Parameters,
+		FRDGBufferRef IndirectArgsBuffer,
+		uint32 IndirectArgOffset)
+	{
+		ClearUnusedGraphResources(ComputeShader, Parameters, { IndirectArgsBuffer });
+
+		GraphBuilder.AddPass(
+			Forward<FRDGEventName>(PassName),
+			Parameters,
+			ERenderGraphPassFlags::Compute,
+			[Parameters, ComputeShader, IndirectArgsBuffer, IndirectArgOffset](FRHICommandList& RHICmdList)
+		{
+			FRHIComputeShader* ShaderRHI = ComputeShader->GetComputeShader();
+			RHICmdList.SetComputeShader(ShaderRHI);
+			SetShaderParameters(RHICmdList, ComputeShader, ShaderRHI, *Parameters);
+			RHICmdList.DispatchIndirectComputeShader( IndirectArgsBuffer->GetIndirectRHICallBuffer(), IndirectArgOffset );
+			UnsetShaderUAVs(RHICmdList, ComputeShader, ShaderRHI);
+		});
+	}
+};
\ No newline at end of file
diff --git a/Engine/Source/Runtime/RenderCore/Public/ShaderParameterStruct.h b/Engine/Source/Runtime/RenderCore/Public/ShaderParameterStruct.h
--- a/Engine/Source/Runtime/RenderCore/Public/ShaderParameterStruct.h
+++ b/Engine/Source/Runtime/RenderCore/Public/ShaderParameterStruct.h
@@ -320,7 +320,7 @@ void SetShaderParameters(FRayTracingShaderBindingsWriter& RTBindingsWriter, cons
 	// Graph SRVs
 	for (const FShaderParameterBindings::FResourceParameter& ParameterBinding : Bindings.GraphSRVs)
 	{
-		auto GraphSRV = *reinterpret_cast<const FRDGTextureSRV* const*>(Base + ParameterBinding.ByteOffset);
+		auto GraphSRV = *reinterpret_cast<const FRDGResource* const*>(Base + ParameterBinding.ByteOffset);
 
 		if (DO_CHECK)
 		{
diff --git a/Engine/Source/Runtime/UtilityShaders/Private/ClearQuad.cpp b/Engine/Source/Runtime/UtilityShaders/Private/ClearQuad.cpp
--- a/Engine/Source/Runtime/UtilityShaders/Private/ClearQuad.cpp
+++ b/Engine/Source/Runtime/UtilityShaders/Private/ClearQuad.cpp
@@ -179,53 +179,56 @@ void ClearUAV(FRHICommandList& RHICmdList, FRHIUnorderedAccessView* Buffer, uint
 }
 
 template< typename T >
-inline void ClearUAV_T(FRHICommandList& RHICmdList, const FSceneRenderTargetItem& RenderTargetItem, const T(&ClearValues)[4])
+inline void ClearUAV_T(FRHICommandList& RHICmdList, FRHITexture* Texture, FRHIUnorderedAccessView* TextureUAV, const T(&ClearValues)[4])
 {
-	if (auto Texture2d = RenderTargetItem.TargetableTexture->GetTexture2D())
+	check( Texture );
+	check( TextureUAV );
+
+	if (auto Texture2d = Texture->GetTexture2D())
 	{
 		TShaderMapRef< FClearTexture2DReplacementCS<T> > ComputeShader(GetGlobalShaderMap(GMaxRHIFeatureLevel));
 		FComputeShaderRHIParamRef ShaderRHI = ComputeShader->GetComputeShader();
 		RHICmdList.SetComputeShader(ShaderRHI);
-		ComputeShader->SetParameters(RHICmdList, RenderTargetItem.UAV, ClearValues);
+		ComputeShader->SetParameters(RHICmdList, TextureUAV, ClearValues);
 		uint32 x = (Texture2d->GetSizeX() + 7) / 8;
 		uint32 y = (Texture2d->GetSizeY() + 7) / 8;
 		RHICmdList.DispatchComputeShader(x, y, 1);
-		ComputeShader->FinalizeParameters(RHICmdList, RenderTargetItem.UAV);
+		ComputeShader->FinalizeParameters(RHICmdList, TextureUAV);
 	}
-	else if (auto Texture2dArray = RenderTargetItem.TargetableTexture->GetTexture2DArray())
+	else if (auto Texture2dArray = Texture->GetTexture2DArray())
 	{
 		TShaderMapRef< FClearTexture2DArrayReplacementCS<T> > ComputeShader(GetGlobalShaderMap(GMaxRHIFeatureLevel));
 		FComputeShaderRHIParamRef ShaderRHI = ComputeShader->GetComputeShader();
 		RHICmdList.SetComputeShader(ShaderRHI);
-		ComputeShader->SetParameters(RHICmdList, RenderTargetItem.UAV, ClearValues);
+		ComputeShader->SetParameters(RHICmdList, TextureUAV, ClearValues);
 		uint32 x = (Texture2dArray->GetSizeX() + 7) / 8;
 		uint32 y = (Texture2dArray->GetSizeY() + 7) / 8;
 		uint32 z = Texture2dArray->GetSizeZ();
 		RHICmdList.DispatchComputeShader(x, y, z);
-		ComputeShader->FinalizeParameters(RHICmdList, RenderTargetItem.UAV);
+		ComputeShader->FinalizeParameters(RHICmdList, TextureUAV);
 	}
-	else if (auto TextureCube = RenderTargetItem.TargetableTexture->GetTextureCube())
+	else if (auto TextureCube = Texture->GetTextureCube())
 	{
 		TShaderMapRef< FClearTexture2DArrayReplacementCS<T> > ComputeShader(GetGlobalShaderMap(GMaxRHIFeatureLevel));
 		FComputeShaderRHIParamRef ShaderRHI = ComputeShader->GetComputeShader();
 		RHICmdList.SetComputeShader(ShaderRHI);
-		ComputeShader->SetParameters(RHICmdList, RenderTargetItem.UAV, ClearValues);
+		ComputeShader->SetParameters(RHICmdList, TextureUAV, ClearValues);
 		uint32 x = (TextureCube->GetSize() + 7) / 8;
 		uint32 y = (TextureCube->GetSize() + 7) / 8;
 		RHICmdList.DispatchComputeShader(x, y, 6);
-		ComputeShader->FinalizeParameters(RHICmdList, RenderTargetItem.UAV);
+		ComputeShader->FinalizeParameters(RHICmdList, TextureUAV);
 	}
-	else if (auto Texture3d = RenderTargetItem.TargetableTexture->GetTexture3D())
+	else if (auto Texture3d = Texture->GetTexture3D())
 	{
 		TShaderMapRef< FClearVolumeReplacementCS<T> > ComputeShader(GetGlobalShaderMap(GMaxRHIFeatureLevel));
 		FComputeShaderRHIParamRef ShaderRHI = ComputeShader->GetComputeShader();
 		RHICmdList.SetComputeShader(ShaderRHI);
-		ComputeShader->SetParameters(RHICmdList, RenderTargetItem.UAV, ClearValues);
+		ComputeShader->SetParameters(RHICmdList, TextureUAV, ClearValues);
 		uint32 x = (Texture3d->GetSizeX() + 3) / 4;
 		uint32 y = (Texture3d->GetSizeY() + 3) / 4;
 		uint32 z = (Texture3d->GetSizeZ() + 3) / 4;
 		RHICmdList.DispatchComputeShader(x, y, z);
-		ComputeShader->FinalizeParameters(RHICmdList, RenderTargetItem.UAV);
+		ComputeShader->FinalizeParameters(RHICmdList, TextureUAV);
 	}
 	else
 	{
@@ -247,17 +250,32 @@ void ClearTexture2DUAV(FRHICommandList& RHICmdList, FUnorderedAccessViewRHIParam
 
 void ClearUAV(FRHICommandList& RHICmdList, const FSceneRenderTargetItem& RenderTargetItem, const float(&ClearValues)[4])
 {
-	ClearUAV_T(RHICmdList, RenderTargetItem, ClearValues);
+	ClearUAV_T(RHICmdList, RenderTargetItem.TargetableTexture, RenderTargetItem.UAV, ClearValues);
 }
 
 void ClearUAV(FRHICommandList& RHICmdList, const FSceneRenderTargetItem& RenderTargetItem, const uint32(&ClearValues)[4])
 {
-	ClearUAV_T(RHICmdList, RenderTargetItem, ClearValues);
+	ClearUAV_T(RHICmdList, RenderTargetItem.TargetableTexture, RenderTargetItem.UAV, ClearValues);
 }
 
 void ClearUAV(FRHICommandList& RHICmdList, const FSceneRenderTargetItem& RenderTargetItem, const FLinearColor& ClearColor)
 {
-	ClearUAV_T(RHICmdList, RenderTargetItem, reinterpret_cast<const float(&)[4]>(ClearColor));
+	ClearUAV_T(RHICmdList, RenderTargetItem.TargetableTexture, RenderTargetItem.UAV, reinterpret_cast<const float(&)[4]>(ClearColor));
+}
+
+void ClearUAV(FRHICommandList& RHICmdList, FRHITexture* Texture, FRHIUnorderedAccessView* TextureUAV, const float(&ClearValues)[4])
+{
+	ClearUAV_T(RHICmdList, Texture, TextureUAV, ClearValues);
+}
+
+void ClearUAV(FRHICommandList& RHICmdList, FRHITexture* Texture, FRHIUnorderedAccessView* TextureUAV, const uint32(&ClearValues)[4])
+{
+	ClearUAV_T(RHICmdList, Texture, TextureUAV, ClearValues);
+}
+
+void ClearUAV(FRHICommandList& RHICmdList, FRHITexture* Texture, FRHIUnorderedAccessView* TextureUAV, const FLinearColor& ClearColor)
+{
+	ClearUAV_T(RHICmdList, Texture, TextureUAV, reinterpret_cast<const float(&)[4]>(ClearColor));
 }
 
 void DrawClearQuadMRT(FRHICommandList& RHICmdList, bool bClearColor, int32 NumClearColors, const FLinearColor* ClearColorArray, bool bClearDepth, float Depth, bool bClearStencil, uint32 Stencil)
diff --git a/Engine/Source/Runtime/UtilityShaders/Public/ClearQuad.h b/Engine/Source/Runtime/UtilityShaders/Public/ClearQuad.h
--- a/Engine/Source/Runtime/UtilityShaders/Public/ClearQuad.h
+++ b/Engine/Source/Runtime/UtilityShaders/Public/ClearQuad.h
@@ -48,9 +48,13 @@ extern UTILITYSHADERS_API void ClearUAV(FRHICommandList& RHICmdList, const FRWBu
 extern UTILITYSHADERS_API void ClearUAV(FRHICommandList& RHICmdList, const FRWBuffer& Buffer, uint32 Value);
 extern UTILITYSHADERS_API void ClearUAV(FRHICommandList& RHICmdList, FRHIUnorderedAccessView* Buffer, uint32 NumBytes, uint32 Value);
 
+extern UTILITYSHADERS_API void ClearUAV(FRHICommandList& RHICmdList, FRHITexture* Texture, FRHIUnorderedAccessView* TextureUAV, const float(&ClearValues)[4]);
+extern UTILITYSHADERS_API void ClearUAV(FRHICommandList& RHICmdList, FRHITexture* Texture, FRHIUnorderedAccessView* TextureUAV, const uint32(&ClearValues)[4]);
+extern UTILITYSHADERS_API void ClearUAV(FRHICommandList& RHICmdList, FRHITexture* Texture, FRHIUnorderedAccessView* TextureUAV, const FLinearColor& ClearColor);
+
 extern UTILITYSHADERS_API void ClearUAV(FRHICommandList& RHICmdList, const FSceneRenderTargetItem& RenderTargetItem, const float(&ClearValues)[4]);
 extern UTILITYSHADERS_API void ClearUAV(FRHICommandList& RHICmdList, const FSceneRenderTargetItem& RenderTargetItem, const uint32(&ClearValues)[4]);
-extern UTILITYSHADERS_API void ClearUAV(FRHICommandList& RHICmdList, const FSceneRenderTargetItem& RenderTargetItem,  const FLinearColor& ClearColor);
+extern UTILITYSHADERS_API void ClearUAV(FRHICommandList& RHICmdList, const FSceneRenderTargetItem& RenderTargetItem, const FLinearColor& ClearColor);
 
 extern UTILITYSHADERS_API void ClearTexture2DUAV(FRHICommandList& RHICmdList, FRHIUnorderedAccessView* UAV, int32 Width, int32 Height, const FLinearColor& ClearColor);
 
diff --git a/Engine/Source/Runtime/Windows/D3D11RHI/Private/D3D11StructuredBuffer.cpp b/Engine/Source/Runtime/Windows/D3D11RHI/Private/D3D11StructuredBuffer.cpp
--- a/Engine/Source/Runtime/Windows/D3D11RHI/Private/D3D11StructuredBuffer.cpp
+++ b/Engine/Source/Runtime/Windows/D3D11RHI/Private/D3D11StructuredBuffer.cpp
@@ -76,6 +76,11 @@ FStructuredBufferRHIRef FD3D11DynamicRHI::RHICreateStructuredBuffer(uint32 Strid
 	TRefCountPtr<ID3D11Buffer> StructuredBufferResource;
 	VERIFYD3D11RESULT_EX(Direct3DDevice->CreateBuffer(&Desc,pInitData,StructuredBufferResource.GetInitReference()), Direct3DDevice);
 
+	if( CreateInfo.DebugName )
+	{
+		StructuredBufferResource->SetPrivateData(WKPDID_D3DDebugObjectName, FCString::Strlen(CreateInfo.DebugName) + 1, TCHAR_TO_ANSI(CreateInfo.DebugName));
+	}
+
 	UpdateBufferStats(StructuredBufferResource, true);
 
 	if(CreateInfo.ResourceArray)
