DR - Fix long running task timer queries crashing in testbuilds.

diff --git a/Engine/Source/Runtime/UtilityShaders/Private/LongGPUTask.cpp b/Engine/Source/Runtime/UtilityShaders/Private/LongGPUTask.cpp
--- a/Engine/Source/Runtime/UtilityShaders/Private/LongGPUTask.cpp
+++ b/Engine/Source/Runtime/UtilityShaders/Private/LongGPUTask.cpp
@@ -60,11 +60,8 @@ void MeasureLongGPUTaskExecutionTime(FRHICommandListImmediate& RHICmdList)
 {
 	const int32 NumIterationsForMeasurement = 5;
 
-	FRenderQueryRHIRef TimeQueryStart;
-	FRenderQueryRHIRef TimeQueryEnd;
-
-	TimeQueryStart = RHICreateRenderQuery(RQT_AbsoluteTime);
-	TimeQueryEnd = RHICreateRenderQuery(RQT_AbsoluteTime);
+	FRenderQueryRHIRef TimeQueryStart = RHICreateRenderQuery(RQT_AbsoluteTime);
+	FRenderQueryRHIRef TimeQueryEnd = RHICreateRenderQuery(RQT_AbsoluteTime);
 
 	if (TimeQueryStart == nullptr || TimeQueryEnd == nullptr)
 	{
@@ -74,19 +71,21 @@ void MeasureLongGPUTaskExecutionTime(FRHICommandListImmediate& RHICmdList)
 		return;
 	}
 
+	uint64 StartTime = 0;
+	uint64 EndTime = 0;
+
+	bool StartQ = RHICmdList.GetRenderQueryResult(TimeQueryStart, StartTime, true);
 	RHICmdList.EndRenderQuery(TimeQueryStart);
 
 	IssueScalableLongGPUTask(RHICmdList, NumIterationsForMeasurement);
 
+	bool EndQ = RHICmdList.GetRenderQueryResult(TimeQueryEnd, EndTime, true);
 	RHICmdList.EndRenderQuery(TimeQueryEnd);
 
 	// Required by DX12 to resolve the query
 	RHICmdList.SubmitCommandsHint();
 
-	uint64 StartTime = 0;
-	uint64 EndTime = 0;
-
-	if (RHIGetRenderQueryResult(TimeQueryStart, StartTime, true) && RHIGetRenderQueryResult(TimeQueryEnd, EndTime, true))
+	if (StartQ && EndQ)
 	{
 		NumMeasuredIterationsToAchieve500ms = FMath::Clamp(FMath::FloorToInt(500.0f / ((EndTime - StartTime) / 1000.0f / NumIterationsForMeasurement)), 1, 200);
 	}
