DR - so the last checkin made it work in Test but crashed everywhere else. After staring at the code a bit more there was another case where RHICmdList.CreateRenderQuery was changed to RHICreateRenderQuery (as before was RHICmdList.GetRenderQueryResult had been changed to RHIGetRenderQueryResult) Those RHI*Minefield Commands (the ones without the CmdList) grab the Commandlist from some *other* Immediate context and that is not gonna end well. I am wondering how many other issues like this might be hidden somewhere.

diff --git a/Engine/Source/Runtime/UtilityShaders/Private/LongGPUTask.cpp b/Engine/Source/Runtime/UtilityShaders/Private/LongGPUTask.cpp
--- a/Engine/Source/Runtime/UtilityShaders/Private/LongGPUTask.cpp
+++ b/Engine/Source/Runtime/UtilityShaders/Private/LongGPUTask.cpp
@@ -60,8 +60,8 @@ void MeasureLongGPUTaskExecutionTime(FRHICommandListImmediate& RHICmdList)
 {
 	const int32 NumIterationsForMeasurement = 5;
 
-	FRenderQueryRHIRef TimeQueryStart = RHICreateRenderQuery(RQT_AbsoluteTime);
-	FRenderQueryRHIRef TimeQueryEnd = RHICreateRenderQuery(RQT_AbsoluteTime);
+	FRenderQueryRHIRef TimeQueryStart = RHICmdList.CreateRenderQuery(RQT_AbsoluteTime);
+	FRenderQueryRHIRef TimeQueryEnd = RHICmdList.CreateRenderQuery(RQT_AbsoluteTime);
 
 	if (TimeQueryStart == nullptr || TimeQueryEnd == nullptr)
 	{
@@ -74,18 +74,18 @@ void MeasureLongGPUTaskExecutionTime(FRHICommandListImmediate& RHICmdList)
 	uint64 StartTime = 0;
 	uint64 EndTime = 0;
 
-	bool StartQ = RHICmdList.GetRenderQueryResult(TimeQueryStart, StartTime, true);
 	RHICmdList.EndRenderQuery(TimeQueryStart);
+	bool HasStartQuery = RHICmdList.GetRenderQueryResult(TimeQueryStart, StartTime, true);
 
 	IssueScalableLongGPUTask(RHICmdList, NumIterationsForMeasurement);
 
-	bool EndQ = RHICmdList.GetRenderQueryResult(TimeQueryEnd, EndTime, true);
 	RHICmdList.EndRenderQuery(TimeQueryEnd);
+	bool HasEndQuery = RHICmdList.GetRenderQueryResult(TimeQueryEnd, EndTime, true);
 
 	// Required by DX12 to resolve the query
 	RHICmdList.SubmitCommandsHint();
 
-	if (StartQ && EndQ)
+	if (HasStartQuery && HasEndQuery)
 	{
 		NumMeasuredIterationsToAchieve500ms = FMath::Clamp(FMath::FloorToInt(500.0f / ((EndTime - StartTime) / 1000.0f / NumIterationsForMeasurement)), 1, 200);
 	}
