DR - Add flush to RHI THread to make sure the query is submitted.

diff --git a/Engine/Source/Runtime/UtilityShaders/Private/LongGPUTask.cpp b/Engine/Source/Runtime/UtilityShaders/Private/LongGPUTask.cpp
--- a/Engine/Source/Runtime/UtilityShaders/Private/LongGPUTask.cpp
+++ b/Engine/Source/Runtime/UtilityShaders/Private/LongGPUTask.cpp
@@ -75,17 +75,16 @@ void MeasureLongGPUTaskExecutionTime(FRHICommandListImmediate& RHICmdList)
 	uint64 EndTime = 0;
 
 	RHICmdList.EndRenderQuery(TimeQueryStart);
-	bool HasStartQuery = RHICmdList.GetRenderQueryResult(TimeQueryStart, StartTime, true);
 
 	IssueScalableLongGPUTask(RHICmdList, NumIterationsForMeasurement);
 
 	RHICmdList.EndRenderQuery(TimeQueryEnd);
-	bool HasEndQuery = RHICmdList.GetRenderQueryResult(TimeQueryEnd, EndTime, true);
 
 	// Required by DX12 to resolve the query
 	RHICmdList.SubmitCommandsHint();
+	RHICmdList.ImmediateFlush(EImmediateFlushType::FlushRHIThread);
 
-	if (HasStartQuery && HasEndQuery)
+	if (RHICmdList.GetRenderQueryResult(TimeQueryStart, StartTime, true) && RHICmdList.GetRenderQueryResult(TimeQueryEnd, EndTime, true))
 	{
 		NumMeasuredIterationsToAchieve500ms = FMath::Clamp(FMath::FloorToInt(500.0f / ((EndTime - StartTime) / 1000.0f / NumIterationsForMeasurement)), 1, 200);
 	}
